Adding ContextualAnimation plugin [experimental].
This initial commits includes:
- ContextualAnimAsset: DataAsset container for the list of animations, alignment tracks for each animation etc.
- ContextualAnimComponent: Component with helper functions to query the DataAsset and draw debug features to visualize entry points and tolerance in the world.
- ContextualAnimEdMode: A very early version of a custom EdMode to tests interactions in Simulate Mode.

diff --git a/Engine/Plugins/Experimental/Animation/ContextualAnimation/ContextualAnimation.uplugin b/Engine/Plugins/Experimental/Animation/ContextualAnimation/ContextualAnimation.uplugin
new file mode 100644
--- /dev/null
+++ b/Engine/Plugins/Experimental/Animation/ContextualAnimation/ContextualAnimation.uplugin
@@ -0,0 +1,36 @@
+{
+	"FileVersion": 3,
+	"Version": 1,
+	"VersionName": "0.1",
+	"FriendlyName": "ContextualAnimation",
+	"Description": "",
+	"Category": "Animation",
+	"CreatedBy": "Epic Games, Inc.",
+	"CreatedByURL": "http://epicgames.com",
+	"DocsURL": "",
+	"MarketplaceURL": "",
+	"SupportURL": "",
+	"EnabledByDefault": false,
+	"CanContainContent": true,
+	"IsBetaVersion": false,
+	"IsExperimentalVersion": true,
+	"Installed": false,
+	"Modules": [
+		{
+			"Name": "ContextualAnimation",
+			"Type": "Runtime",
+			"LoadingPhase": "Default"
+		},
+		{
+			"Name": "ContextualAnimationEditor",
+			"Type": "Editor",
+			"LoadingPhase": "Default"
+		}
+	],
+	"Plugins": [
+		{
+			"Name": "MotionWarping",
+			"Enabled": true
+		}
+	]
+}
\ No newline at end of file
diff --git a/Engine/Plugins/Experimental/Animation/ContextualAnimation/Source/ContextualAnimation/ContextualAnimation.Build.cs b/Engine/Plugins/Experimental/Animation/ContextualAnimation/Source/ContextualAnimation/ContextualAnimation.Build.cs
new file mode 100644
--- /dev/null
+++ b/Engine/Plugins/Experimental/Animation/ContextualAnimation/Source/ContextualAnimation/ContextualAnimation.Build.cs
@@ -0,0 +1,51 @@
+// Copyright Epic Games, Inc. All Rights Reserved.
+
+using UnrealBuildTool;
+
+public class ContextualAnimation : ModuleRules
+{
+	public ContextualAnimation(ReadOnlyTargetRules Target) : base(Target)
+	{
+		PCHUsage = ModuleRules.PCHUsageMode.UseExplicitOrSharedPCHs;
+		
+		PublicIncludePaths.AddRange(
+			new string[] {
+				// ... add public include paths required here ...
+			}
+			);
+				
+		
+		PrivateIncludePaths.AddRange(
+			new string[] {
+				// ... add other private include paths required here ...
+			}
+			);
+			
+		
+		PublicDependencyModuleNames.AddRange(
+			new string[]
+			{
+				"Core",
+				"CoreUObject",
+                "Engine"
+			}
+			);
+			
+		
+		PrivateDependencyModuleNames.AddRange(
+			new string[]
+			{
+				"Slate",
+				"SlateCore"
+            }
+			);
+		
+		
+		DynamicallyLoadedModuleNames.AddRange(
+			new string[]
+			{
+				// ... add any modules that your module loads dynamically here ...
+			}
+			);
+	}
+}
diff --git a/Engine/Plugins/Experimental/Animation/ContextualAnimation/Source/ContextualAnimation/Private/ContextualAnimAsset.cpp b/Engine/Plugins/Experimental/Animation/ContextualAnimation/Source/ContextualAnimation/Private/ContextualAnimAsset.cpp
new file mode 100644
--- /dev/null
+++ b/Engine/Plugins/Experimental/Animation/ContextualAnimation/Source/ContextualAnimation/Private/ContextualAnimAsset.cpp
@@ -0,0 +1,197 @@
+// Copyright Epic Games, Inc. All Rights Reserved.
+
+#include "ContextualAnimAsset.h"
+#include "Components/SkeletalMeshComponent.h"
+#include "DrawDebugHelpers.h"
+#include "AnimationUtils.h"
+#include "AnimationRuntime.h"
+#include "Animation/AnimMontage.h"
+#include "Animation/AnimationPoseData.h"
+#include "Animation/AnimTypes.h"
+
+DEFINE_LOG_CATEGORY(LogContextualAnim);
+
+// UContextualAnimUtilities
+///////////////////////////////////////////////////////////////////////
+
+void UContextualAnimUtilities::ExtractLocalSpacePose(const UAnimSequenceBase* Animation, const FBoneContainer& BoneContainer, float Time, bool bExtractRootMotion, FCompactPose& OutPose)
+{
+	OutPose.SetBoneContainer(&BoneContainer);
+
+	FBlendedCurve Curve;
+	Curve.InitFrom(BoneContainer);
+
+	FAnimExtractContext Context(Time, bExtractRootMotion);
+
+	FStackCustomAttributes Attributes;
+	FAnimationPoseData AnimationPoseData(OutPose, Curve, Attributes);
+	if (const UAnimSequence* AnimSequence = Cast<UAnimSequence>(Animation))
+	{
+		AnimSequence->GetBonePose(AnimationPoseData, Context);
+	}
+	else if (const UAnimMontage* AnimMontage = Cast<UAnimMontage>(Animation))
+	{
+		const FAnimTrack& AnimTrack = AnimMontage->SlotAnimTracks[0].AnimTrack;
+		AnimTrack.GetAnimationPose(AnimationPoseData, Context);
+	}
+}
+
+void UContextualAnimUtilities::ExtractComponentSpacePose(const UAnimSequenceBase* Animation, const FBoneContainer& BoneContainer, float Time, bool bExtractRootMotion, FCSPose<FCompactPose>& OutPose)
+{
+	FCompactPose Pose;
+	ExtractLocalSpacePose(Animation, BoneContainer, Time, bExtractRootMotion, Pose);
+	OutPose.InitPose(MoveTemp(Pose));
+}
+
+FTransform UContextualAnimUtilities::ExtractRootMotionFromAnimation(const UAnimSequenceBase* Animation, float StartTime, float EndTime)
+{
+	if (const UAnimMontage* Anim = Cast<UAnimMontage>(Animation))
+	{
+		return Anim->ExtractRootMotionFromTrackRange(StartTime, EndTime);
+	}
+
+	if (const UAnimSequence* Anim = Cast<UAnimSequence>(Animation))
+	{
+		return Anim->ExtractRootMotionFromRange(StartTime, EndTime);
+	}
+
+	return FTransform::Identity;
+}
+
+// FAlignmentTrackContainer
+///////////////////////////////////////////////////////////////////////
+
+FTransform FAlignmentTrackContainer::ExtractTransformAtTime(float Time) const
+{
+	FTransform AlignmentTransform = FTransform::Identity;
+	const int32 TotalFrames = Track.PosKeys.Num();
+	const float TrackLength = (TotalFrames - 1) * SampleRate;
+	FAnimationUtils::ExtractTransformFromTrack(Time, TotalFrames, TrackLength, Track, EAnimInterpolationType::Linear, AlignmentTransform);
+	return AlignmentTransform;
+}
+
+// UContextualAnimAsset
+///////////////////////////////////////////////////////////////////////
+
+UContextualAnimAsset::UContextualAnimAsset(const FObjectInitializer& ObjectInitializer)
+	: Super(ObjectInitializer) 
+{
+	MeshToComponent = FTransform(FRotator(0.f, -90.f, 0.f));
+}
+
+#if WITH_EDITOR
+void UContextualAnimAsset::PostEditChangeProperty(FPropertyChangedEvent& PropertyChangedEvent)
+{
+	Super::PostEditChangeProperty(PropertyChangedEvent);
+}
+#endif // WITH_EDITOR
+
+void UContextualAnimAsset::PreSave(const class ITargetPlatform* TargetPlatform)
+{
+	const FTransform MeshToComponentInverse = MeshToComponent.Inverse();
+
+	TArray<FBoneIndexType> RequiredBoneIndexArray;
+	RequiredBoneIndexArray.Add(0);
+	FBoneContainer BoneContainer;
+
+	for (FContextualAnimData& Data : DataContainer)
+	{
+		const UAnimMontage* Animation = Data.Animation.LoadSynchronous();
+
+		if(Animation == nullptr || !Animation->IsValidSectionIndex(0))
+		{
+			continue;
+		}
+		
+		// For now we just assume the entry animation for the interaction is the first section of the montage
+		float SectionStartTime, SectionEndTime;
+		Animation->GetSectionStartAndEndTime(0, SectionStartTime, SectionEndTime);
+
+		const float SampleRate = 1 / 30.f;
+		float Time = 0.f;
+		float EndTime = SectionEndTime;
+		int32 SampleIndex = 0;
+
+		Data.SyncTime = EndTime;
+		Data.AlignmentData.Initialize(SampleRate);
+		FRawAnimSequenceTrack& NewTrack = Data.AlignmentData.Track;
+
+		if (AlignmentJoint.IsNone())
+		{
+			if(!BoneContainer.IsValid())
+			{
+				BoneContainer = FBoneContainer(RequiredBoneIndexArray, FCurveEvaluationOption(false), *Animation->GetSkeleton());
+			}
+
+			while (Time < EndTime)
+			{
+				Time = FMath::Clamp(SampleIndex * SampleRate, 0.f, EndTime);
+				SampleIndex++;
+
+				FCSPose<FCompactPose> ComponentSpacePose;
+				UContextualAnimUtilities::ExtractComponentSpacePose(Animation, BoneContainer, Time, false, ComponentSpacePose);
+
+				const FTransform RootTransform = MeshToComponentInverse * ComponentSpacePose.GetComponentSpaceTransform(FCompactPoseBoneIndex(0));
+				const FTransform RootRelativeToRef = RootTransform.GetRelativeTransform(MeshToComponentInverse);
+
+				NewTrack.PosKeys.Add(RootRelativeToRef.GetLocation());
+				NewTrack.RotKeys.Add(RootRelativeToRef.GetRotation());
+				NewTrack.ScaleKeys.Add(RootRelativeToRef.GetScale3D());
+			};
+		}
+		else
+		{
+			if (!BoneContainer.IsValid())
+			{
+				const int32 AlignmentJointIdx = Animation->GetSkeleton()->GetReferenceSkeleton().FindBoneIndex(AlignmentJoint);
+				if (AlignmentJointIdx == INDEX_NONE)
+				{
+					UE_LOG(LogContextualAnim, Warning, TEXT("%s Can't generate alignment track for %s. Reason: AlignmentJoint (%s) invalid"),
+						*GetNameSafe(this), *GetNameSafe(Animation), *AlignmentJoint.ToString());
+
+					continue;
+				}
+
+				RequiredBoneIndexArray.Add(AlignmentJointIdx);
+				Animation->GetSkeleton()->GetReferenceSkeleton().EnsureParentsExistAndSort(RequiredBoneIndexArray);
+
+				BoneContainer = FBoneContainer(RequiredBoneIndexArray, FCurveEvaluationOption(false), *Animation->GetSkeleton());
+			}
+
+			while (Time < EndTime)
+			{
+				Time = FMath::Clamp(SampleIndex * SampleRate, 0.f, EndTime);
+				SampleIndex++;
+
+				FCSPose<FCompactPose> ComponentSpacePose;
+				UContextualAnimUtilities::ExtractComponentSpacePose(Animation, BoneContainer, Time, false, ComponentSpacePose);
+
+				const FTransform RootTransform = MeshToComponentInverse * ComponentSpacePose.GetComponentSpaceTransform(FCompactPoseBoneIndex(0));
+
+				FCompactPoseBoneIndex BoneIndex = BoneContainer.MakeCompactPoseIndex(FMeshPoseBoneIndex(BoneContainer.GetReferenceSkeleton().FindBoneIndex(AlignmentJoint)));
+				const FTransform RefTransform = MeshToComponentInverse * ComponentSpacePose.GetComponentSpaceTransform(FCompactPoseBoneIndex(BoneIndex));
+
+				const FTransform RootRelativeToRef = RootTransform.GetRelativeTransform(RefTransform);
+
+				NewTrack.PosKeys.Add(RootRelativeToRef.GetLocation());
+				NewTrack.RotKeys.Add(RootRelativeToRef.GetRotation());
+				NewTrack.ScaleKeys.Add(RootRelativeToRef.GetScale3D());
+			};
+		}
+
+		const FTransform EntryTransform = Data.GetAlignmentTransformAtEntryTime();
+
+		// Reference value for Distance Test
+		Data.Distance.Value = EntryTransform.GetTranslation().Size2D();
+
+		// Reference value for Angle Test
+		const FVector Direction = EntryTransform.GetTranslation().GetSafeNormal2D();
+		FQuat Delta = FQuat::FindBetweenNormals(Direction, FVector::ForwardVector);
+		Data.Angle.Value = Delta.Rotator().Yaw;
+
+		// Reference value for Facing Test
+		Data.Facing.Value = EntryTransform.Rotator().Yaw;
+	}
+
+	Super::PreSave(TargetPlatform);
+}
\ No newline at end of file
diff --git a/Engine/Plugins/Experimental/Animation/ContextualAnimation/Source/ContextualAnimation/Private/ContextualAnimComponent.cpp b/Engine/Plugins/Experimental/Animation/ContextualAnimation/Source/ContextualAnimation/Private/ContextualAnimComponent.cpp
new file mode 100644
--- /dev/null
+++ b/Engine/Plugins/Experimental/Animation/ContextualAnimation/Source/ContextualAnimation/Private/ContextualAnimComponent.cpp
@@ -0,0 +1,325 @@
+// Copyright Epic Games, Inc. All Rights Reserved.
+
+#include "ContextualAnimComponent.h"
+#include "ContextualAnimAsset.h"
+#include "DrawDebugHelpers.h"
+
+UContextualAnimComponent::UContextualAnimComponent(const FObjectInitializer& ObjectInitializer)
+	: Super(ObjectInitializer)
+{
+	PrimaryComponentTick.bCanEverTick = true;
+	PrimaryComponentTick.bStartWithTickEnabled = false;
+}
+
+void UContextualAnimComponent::TickComponent(float DeltaTime, enum ELevelTick TickType, FActorComponentTickFunction* ThisTickFunction)
+{
+	Super::TickComponent(DeltaTime,TickType, ThisTickFunction);
+}
+
+const FContextualAnimData* UContextualAnimComponent::FindBestDataForActor(const AActor* Querier) const
+{
+	if(ContextualAnimAsset == nullptr || Querier == nullptr)
+	{
+		return nullptr;
+	}
+
+	for (int32 Idx = 0; Idx < ContextualAnimAsset->DataContainer.Num(); Idx++)
+	{
+		const FContextualAnimData& Data = ContextualAnimAsset->DataContainer[Idx];
+		const FTransform EntryTransform = Data.GetAlignmentTransformAtEntryTime() * GetComponentTransform();
+		
+		FVector Origin = GetComponentTransform().GetLocation();
+		FVector Direction = (EntryTransform.GetLocation() - Origin).GetSafeNormal2D();
+
+		if (Data.OffsetFromOrigin != 0.f)
+		{
+			Origin = Origin + Direction * Data.OffsetFromOrigin;
+		}
+
+		// Distance Test
+		//--------------------------------------------------
+		if(Data.Distance.MaxDistance > 0.f || Data.Distance.MinDistance > 0.f)
+		{
+			const float DistSq = FVector::DistSquared2D(Origin, Querier->GetActorLocation());
+
+			if (Data.Distance.MaxDistance > 0.f)
+			{
+				if (DistSq > FMath::Square(Data.Distance.MaxDistance))
+				{
+					continue;
+				}
+			}
+
+			if (Data.Distance.MinDistance > 0.f)
+			{
+				if (DistSq < FMath::Square(Data.Distance.MinDistance))
+				{
+					continue;
+				}
+			}
+		}
+
+		// Angle Test
+		//--------------------------------------------------
+		if(Data.Angle.Tolerance > 0.f)
+		{
+			//@TODO: Cache this
+			const float AngleCos = FMath::Cos(FMath::Clamp(FMath::DegreesToRadians(Data.Angle.Tolerance), 0.f, PI));
+			const FVector ToLocation = (Querier->GetActorLocation() - Origin).GetSafeNormal2D();
+			if (FVector::DotProduct(ToLocation, Direction) < AngleCos)
+			{
+				continue;
+			}
+		}
+
+		// Facing Test
+		//--------------------------------------------------
+		if(Data.Facing.Tolerance > 0.f)
+		{
+			//@TODO: Cache this
+			const float FacingCos = FMath::Cos(FMath::Clamp(FMath::DegreesToRadians(Data.Facing.Tolerance), 0.f, PI));
+			if (FVector::DotProduct(Querier->GetActorForwardVector(), EntryTransform.GetRotation().GetForwardVector()) < FacingCos)
+			{
+				continue;
+			}
+		}
+
+		// Return the first item that passes all tests
+		return &Data;
+	}
+
+	return nullptr;
+}
+
+const FContextualAnimData* UContextualAnimComponent::FindClosestDataForActor(const AActor* Querier) const
+{
+	if (ContextualAnimAsset == nullptr || Querier == nullptr)
+	{
+		return nullptr;
+	}
+
+	float BestDistanceSq = MAX_FLT;
+	const FContextualAnimData* BestData = nullptr;
+	for (int32 Idx = 0; Idx < ContextualAnimAsset->DataContainer.Num(); Idx++)
+	{
+		const FContextualAnimData& Data = ContextualAnimAsset->DataContainer[Idx];
+		const FTransform EntryTransform = Data.GetAlignmentTransformAtEntryTime() * GetComponentTransform();
+		const float DistSq = FVector::DistSquared2D(EntryTransform.GetLocation(), Querier->GetActorLocation());
+		if (DistSq < BestDistanceSq)
+		{
+			BestDistanceSq = DistSq;
+			BestData = &Data;
+		}
+	}
+
+	return BestData;
+}
+
+bool UContextualAnimComponent::FindBestEntryPointForActor(const AActor* Querier, FContextualAnimEntryPoint& Result) const
+{
+	if(const FContextualAnimData* Data = FindBestDataForActor(Querier))
+	{
+		Result.Animation = Data->Animation;
+		Result.EntryTransform = Data->GetAlignmentTransformAtEntryTime() * GetComponentTransform();
+		Result.SyncTransform = Data->GetAlignmentTransformAtSyncTime() * GetComponentTransform();
+		return true;
+	}
+
+	return false;
+}
+
+bool UContextualAnimComponent::FindClosestEntryPointForActor(const AActor* Querier, FContextualAnimEntryPoint& Result) const
+{
+	if (const FContextualAnimData* Data = FindClosestDataForActor(Querier))
+	{
+		Result.Animation = Data->Animation;
+		Result.EntryTransform = Data->GetAlignmentTransformAtEntryTime() * GetComponentTransform();
+		Result.SyncTransform = Data->GetAlignmentTransformAtSyncTime() * GetComponentTransform();
+		return true;
+	}
+
+	return false;
+}
+
+FBoxSphereBounds UContextualAnimComponent::CalcBounds(const FTransform& LocalToWorld) const
+{
+	FBox BB(FVector(0.f, 0.f, -10.f), FVector(0.f, 0.f, 10.f));
+	if (ContextualAnimAsset)
+	{
+		for (const FContextualAnimData& Data : ContextualAnimAsset->DataContainer)
+		{
+			BB += Data.GetAlignmentTransformAtEntryTime().GetLocation();
+		}
+	}
+	return FBoxSphereBounds(BB.TransformBy(GetComponentTransform()));
+}
+
+FPrimitiveSceneProxy* UContextualAnimComponent::CreateSceneProxy()
+{
+	class FContextualAnimSceneProxy final : public FPrimitiveSceneProxy
+	{
+	public:
+
+		SIZE_T GetTypeHash() const override
+		{
+			static size_t UniquePointer;
+			return reinterpret_cast<size_t>(&UniquePointer);
+		}
+
+		FContextualAnimSceneProxy(const UContextualAnimComponent* InComponent)
+			: FPrimitiveSceneProxy(InComponent)
+			, ContextualAnimComp(InComponent)
+		{
+		}
+
+		static void DrawSector(FPrimitiveDrawInterface* PDI, const FVector& Origin, const FVector& Direction, float MinDistance, float MaxDistance, float MinAngle, float MaxAngle, const FLinearColor& Color, uint8 DepthPriority, float Thickness)
+		{
+			// Draw Cone lines
+			const FVector LeftDirection = Direction.RotateAngleAxis(MinAngle, FVector::UpVector);
+			const FVector RightDirection = Direction.RotateAngleAxis(MaxAngle, FVector::UpVector);
+			PDI->DrawLine(Origin + (LeftDirection * MinDistance), Origin + (LeftDirection * MaxDistance), Color, DepthPriority, Thickness);
+			PDI->DrawLine(Origin + (RightDirection * MinDistance), Origin + (RightDirection * MaxDistance), Color, DepthPriority, Thickness);
+
+			// Draw Near Arc
+			FVector LastDirection = LeftDirection;
+			for (float Step = MinAngle; Step <= MaxAngle; Step += 10)
+			{
+				const float Length = MinDistance;
+				const float Angle = FMath::Clamp(Step, MinAngle, MaxAngle);
+				const FVector NewDirection = Direction.RotateAngleAxis(Angle, FVector::UpVector);
+				const FVector LineStart = Origin + (LastDirection * Length);
+				const FVector LineEnd = Origin + (NewDirection * Length);
+				PDI->DrawLine(LineStart, LineEnd, Color, DepthPriority, Thickness);
+				LastDirection = NewDirection;
+			}
+
+			// Draw Far Arc
+			LastDirection = LeftDirection;
+			for (float Step = MinAngle; Step <= MaxAngle; Step += 10)
+			{
+				const float Length = MaxDistance;
+				const float Angle = FMath::Clamp(Step, MinAngle, MaxAngle);
+				const FVector NewDirection = Direction.RotateAngleAxis(Angle, FVector::UpVector);
+				const FVector LineStart = Origin + (LastDirection * Length);
+				const FVector LineEnd = Origin + (NewDirection * Length);
+				PDI->DrawLine(LineStart, LineEnd, Color, DepthPriority, Thickness);
+				LastDirection = NewDirection;
+			}
+		}
+
+		virtual void GetDynamicMeshElements(const TArray<const FSceneView*>& Views, const FSceneViewFamily& ViewFamily, uint32 VisibilityMap, FMeshElementCollector& Collector) const override
+		{
+			const UContextualAnimAsset* Asset = ContextualAnimComp.IsValid() ? ContextualAnimComp->ContextualAnimAsset : nullptr;
+			if (Asset == nullptr)
+			{
+				return;
+			}
+
+			const FContextualAnimDebugParams Params = ContextualAnimComp->DebugParams;
+
+			const FTransform ToWorldTransform = FTransform(GetLocalToWorld());
+
+			for (int32 ViewIndex = 0; ViewIndex < Views.Num(); ViewIndex++)
+			{
+				if (VisibilityMap & (1 << ViewIndex))
+				{
+					FPrimitiveDrawInterface* PDI = Collector.GetPDI(ViewIndex);
+
+					FName BestAssetPathName = NAME_None;
+					if (Params.TestActor.IsValid())
+					{
+						if (const FContextualAnimData* Data = ContextualAnimComp->FindBestDataForActor(Params.TestActor.Get()))
+						{
+							BestAssetPathName = Data->Animation.GetUniqueID().GetAssetPathName();
+						}
+					}
+
+					for(const FContextualAnimData& Data : Asset->DataContainer)
+					{
+						FLinearColor DrawColor = BestAssetPathName == Data.Animation.GetUniqueID().GetAssetPathName() ? FLinearColor::Red : FLinearColor::White;
+
+						// Draw Entry Point
+						const FTransform EntryTransform = (Data.GetAlignmentTransformAtEntryTime() * ToWorldTransform);
+						DrawCoordinateSystem(PDI, EntryTransform.GetLocation(), EntryTransform.Rotator(), 20.f, SDPG_World, 3.f);
+
+						// Draw Facing Tolerance
+						if(Data.Facing.Tolerance > 0.f)
+						{
+							DrawSector(PDI, EntryTransform.GetLocation(), EntryTransform.GetRotation().GetForwardVector(), 0.f, 30.f, -Data.Facing.Tolerance, Data.Facing.Tolerance, DrawColor, SDPG_World, 1.f);
+						}
+						else
+						{
+							DrawCircle(PDI, EntryTransform.GetLocation(), FVector(1, 0, 0), FVector(0, 1, 0), DrawColor, 30.f, 12, SDPG_World, 1.f);
+						}
+
+						// Draw Sync Point
+						const FTransform SyncPoint = Data.GetAlignmentTransformAtSyncTime() * ToWorldTransform;
+						DrawCoordinateSystem(PDI, SyncPoint.GetLocation(), SyncPoint.Rotator(), 20.f, SDPG_World, 3.f);
+
+						if(Params.DrawAlignmentTransformAtTime != Data.EntryTime)
+						{
+							const FTransform RootAtTime = (Data.GetAlignmentTransformAtTime(Params.DrawAlignmentTransformAtTime) * ToWorldTransform);
+							DrawCoordinateSystem(PDI, RootAtTime.GetLocation(), RootAtTime.Rotator(), 10.f, SDPG_World, 2.f);
+						}
+
+						// Draw Sector
+						FVector Origin = ToWorldTransform.GetLocation();
+						FVector Direction = (EntryTransform.GetLocation() - ToWorldTransform.GetLocation()).GetSafeNormal2D();
+
+						if(Data.OffsetFromOrigin != 0.f)
+						{
+							Origin = Origin + Direction * Data.OffsetFromOrigin;
+						}
+
+						if(Data.Angle.Tolerance > 0.f)
+						{
+							DrawSector(PDI, Origin, Direction, Data.Distance.MinDistance, Data.Distance.MaxDistance, -Data.Angle.Tolerance, Data.Angle.Tolerance, DrawColor, SDPG_World, 3.f);
+						}
+						else
+						{
+							if(Data.Distance.MinDistance > 0.f)
+							{
+								DrawCircle(PDI, Origin, FVector(1, 0, 0), FVector(0, 1, 0), DrawColor, Data.Distance.MinDistance, 12, SDPG_World, 2.f);
+							}
+							
+							if (Data.Distance.MaxDistance > 0.f)
+							{
+								DrawCircle(PDI, Origin, FVector(1, 0, 0), FVector(0, 1, 0), DrawColor, Data.Distance.MaxDistance, 12, SDPG_World, 2.f);
+							}
+						}
+					}
+				}
+			}
+		}
+
+		virtual FPrimitiveViewRelevance GetViewRelevance(const FSceneView* View) const override
+		{
+			const bool bShowForCollision = View->Family->EngineShowFlags.Collision;
+			FPrimitiveViewRelevance Result;
+			Result.bDrawRelevance = IsShown(View);
+			Result.bDynamicRelevance = true;
+			Result.bNormalTranslucency = Result.bSeparateTranslucency = IsShown(View);
+			return Result;
+		}
+
+		virtual uint32 GetMemoryFootprint(void) const override
+		{
+			return(sizeof(*this) + GetAllocatedSize());
+		}
+
+		uint32 GetAllocatedSize(void) const
+		{
+			return(FPrimitiveSceneProxy::GetAllocatedSize());
+		}
+
+	private:
+		TWeakObjectPtr<const UContextualAnimComponent> ContextualAnimComp;
+	};
+
+	if(bEnableDebug)
+	{
+		return new FContextualAnimSceneProxy(this);
+	}
+
+	return nullptr;
+}
\ No newline at end of file
diff --git a/Engine/Plugins/Experimental/Animation/ContextualAnimation/Source/ContextualAnimation/Private/ContextualAnimation.cpp b/Engine/Plugins/Experimental/Animation/ContextualAnimation/Source/ContextualAnimation/Private/ContextualAnimation.cpp
new file mode 100644
--- /dev/null
+++ b/Engine/Plugins/Experimental/Animation/ContextualAnimation/Source/ContextualAnimation/Private/ContextualAnimation.cpp
@@ -0,0 +1,19 @@
+// Copyright Epic Games, Inc. All Rights Reserved.
+
+#include "ContextualAnimation.h"
+
+#define LOCTEXT_NAMESPACE "ContextualAnimationModule"
+
+void FContextualAnimationModule::StartupModule()
+{
+
+}
+
+void FContextualAnimationModule::ShutdownModule()
+{
+
+}
+
+#undef LOCTEXT_NAMESPACE
+	
+IMPLEMENT_MODULE(FContextualAnimationModule, ContextualAnimation)
\ No newline at end of file
diff --git a/Engine/Plugins/Experimental/Animation/ContextualAnimation/Source/ContextualAnimation/Public/ContextualAnimAsset.h b/Engine/Plugins/Experimental/Animation/ContextualAnimation/Source/ContextualAnimation/Public/ContextualAnimAsset.h
new file mode 100644
--- /dev/null
+++ b/Engine/Plugins/Experimental/Animation/ContextualAnimation/Source/ContextualAnimation/Public/ContextualAnimAsset.h
@@ -0,0 +1,131 @@
+// Copyright Epic Games, Inc. All Rights Reserved.
+
+#pragma once
+
+#include "Engine/DataAsset.h"
+#include "Animation/AnimSequence.h"
+#include "ContextualAnimAsset.generated.h"
+
+DECLARE_LOG_CATEGORY_EXTERN(LogContextualAnim, Log, All);
+
+UCLASS()
+class UContextualAnimUtilities : public UBlueprintFunctionLibrary
+{
+	GENERATED_BODY()
+
+public:
+
+	static void ExtractLocalSpacePose(const UAnimSequenceBase* Animation, const FBoneContainer& BoneContainer, float Time, bool bExtractRootMotion, FCompactPose& OutPose);
+	static void ExtractComponentSpacePose(const UAnimSequenceBase* Animation, const FBoneContainer& BoneContainer, float Time, bool bExtractRootMotion, FCSPose<FCompactPose>& OutPose);
+	static FTransform ExtractRootMotionFromAnimation(const UAnimSequenceBase* Animation, float StartTime, float EndTime);
+};
+
+USTRUCT()
+struct FAlignmentTrackContainer
+{
+	GENERATED_BODY()
+
+	UPROPERTY()
+	FRawAnimSequenceTrack Track;
+
+	UPROPERTY()
+	float SampleRate;
+
+	FTransform ExtractTransformAtTime(float Time) const;
+
+	void Initialize(float InSampleRate)
+	{
+		Track.PosKeys.Empty();
+		Track.RotKeys.Empty();
+		Track.ScaleKeys.Empty();
+
+		SampleRate = InSampleRate;
+	}
+};
+
+USTRUCT(BlueprintType)
+struct FContextualAnimDistanceParam
+{
+	GENERATED_BODY()
+
+	UPROPERTY(VisibleAnywhere, BlueprintReadWrite, Category = "Defaults", meta = (UIMin = 0, ClampMin = 0))
+	float Value = 0.f;
+
+	UPROPERTY(EditAnywhere, BlueprintReadWrite, Category = "Defaults", meta = (UIMin = 0, ClampMin = 0))
+	float MinDistance = 0.f;
+
+	UPROPERTY(EditAnywhere, BlueprintReadWrite, Category = "Defaults", meta = (UIMin = 0, ClampMin = 0))
+	float MaxDistance = 0.f;
+};
+
+USTRUCT(BlueprintType)
+struct FContextualAnimAngleParam
+{
+	GENERATED_BODY()
+
+	UPROPERTY(VisibleAnywhere, BlueprintReadWrite, Category = "Defaults")
+	float Value = 0.f;
+
+	UPROPERTY(EditAnywhere, BlueprintReadWrite, Category = "Defaults")
+	float Tolerance = 0.f;
+};
+
+USTRUCT(BlueprintType)
+struct FContextualAnimData
+{
+	GENERATED_BODY()
+
+	UPROPERTY(EditAnywhere, BlueprintReadWrite, Category = "Defaults")
+	TSoftObjectPtr<UAnimMontage> Animation;
+
+	UPROPERTY(EditAnywhere, BlueprintReadWrite, Category = "Defaults", meta = (UIMin = 0, ClampMin = 0))
+	float EntryTime = 0.f;
+
+	UPROPERTY(EditAnywhere, BlueprintReadWrite, Category = "Defaults", meta = (UIMin = 0, ClampMin = 0))
+	float SyncTime = 0.f;
+
+	UPROPERTY(EditAnywhere, BlueprintReadWrite, Category = "Defaults")
+	float OffsetFromOrigin = 0.f;
+
+	UPROPERTY(EditAnywhere, BlueprintReadWrite, Category = "Defaults")
+	FContextualAnimDistanceParam Distance;
+
+	UPROPERTY(EditAnywhere, BlueprintReadWrite, Category = "Defaults")
+	FContextualAnimAngleParam Angle;
+
+	UPROPERTY(EditAnywhere, BlueprintReadWrite, Category = "Defaults")
+	FContextualAnimAngleParam Facing;
+
+	UPROPERTY()
+	FAlignmentTrackContainer AlignmentData;
+
+	FORCEINLINE FTransform GetAlignmentTransformAtTime(float Time) const { return AlignmentData.ExtractTransformAtTime(Time); }
+	FORCEINLINE FTransform GetAlignmentTransformAtEntryTime() const { return AlignmentData.ExtractTransformAtTime(EntryTime); }
+	FORCEINLINE FTransform GetAlignmentTransformAtSyncTime() const  { return AlignmentData.ExtractTransformAtTime(SyncTime); }
+};
+
+UCLASS(Blueprintable)
+class CONTEXTUALANIMATION_API UContextualAnimAsset : public UDataAsset
+{
+	GENERATED_BODY()
+
+public:
+
+	UPROPERTY(EditAnywhere, BlueprintReadWrite, Category = "Alignment")
+	FName AlignmentJoint;
+
+	UPROPERTY(EditAnywhere, BlueprintReadWrite, Category = "Alignment")
+	FTransform MeshToComponent;
+
+	UPROPERTY(EditAnywhere, BlueprintReadWrite, Category = "Animations")
+	TArray<FContextualAnimData> DataContainer;
+
+	UContextualAnimAsset(const FObjectInitializer& ObjectInitializer);
+
+	virtual void PreSave(const class ITargetPlatform* TargetPlatform) override;
+
+#if WITH_EDITOR
+	virtual void PostEditChangeProperty(FPropertyChangedEvent& PropertyChangedEvent) override;
+#endif // WITH_EDITOR
+
+};
diff --git a/Engine/Plugins/Experimental/Animation/ContextualAnimation/Source/ContextualAnimation/Public/ContextualAnimComponent.h b/Engine/Plugins/Experimental/Animation/ContextualAnimation/Source/ContextualAnimation/Public/ContextualAnimComponent.h
new file mode 100644
--- /dev/null
+++ b/Engine/Plugins/Experimental/Animation/ContextualAnimation/Source/ContextualAnimation/Public/ContextualAnimComponent.h
@@ -0,0 +1,71 @@
+// Copyright Epic Games, Inc. All Rights Reserved.
+
+#pragma once
+
+#include "CoreMinimal.h"
+#include "Components/PrimitiveComponent.h"
+#include "ContextualAnimComponent.generated.h"
+
+class UContextualAnimAsset;
+struct FContextualAnimData;
+
+USTRUCT(BlueprintType)
+struct FContextualAnimEntryPoint
+{
+	GENERATED_BODY()
+
+	UPROPERTY(EditAnywhere, BlueprintReadWrite, Category = "Defaults")
+	TSoftObjectPtr<UAnimMontage> Animation;
+
+	UPROPERTY(EditAnywhere, BlueprintReadWrite, Category = "Defaults")
+	FTransform EntryTransform;
+
+	UPROPERTY(EditAnywhere, BlueprintReadWrite, Category = "Defaults")
+	FTransform SyncTransform;
+};
+
+USTRUCT(BlueprintType)
+struct FContextualAnimDebugParams
+{
+	GENERATED_BODY()
+
+	UPROPERTY(EditAnywhere, BlueprintReadWrite, Category = "Debug")
+	TWeakObjectPtr<AActor> TestActor;
+
+	UPROPERTY(EditAnywhere, BlueprintReadWrite, Category = "Debug", meta = (UIMin = 0, ClampMin = 0))
+	float DrawAlignmentTransformAtTime = 0.f;
+};
+
+UCLASS(meta = (BlueprintSpawnableComponent))
+class CONTEXTUALANIMATION_API UContextualAnimComponent : public UPrimitiveComponent
+{
+	GENERATED_BODY()
+
+public:
+
+	UContextualAnimComponent(const FObjectInitializer& ObjectInitializer);
+
+	UPROPERTY(EditAnywhere, BlueprintReadWrite, Category = "Config")
+	UContextualAnimAsset* ContextualAnimAsset;
+
+	UPROPERTY(EditAnywhere, BlueprintReadWrite, Category = "Debug")
+	bool bEnableDebug;
+
+	UPROPERTY(EditAnywhere, BlueprintReadWrite, Category = "Debug", meta = (ShowOnlyInnerProperties, EditCondition = "bEnableDebug"))
+	FContextualAnimDebugParams DebugParams;
+
+	virtual void TickComponent(float DeltaTime, enum ELevelTick TickType, FActorComponentTickFunction* ThisTickFunction) override;
+
+	virtual FPrimitiveSceneProxy* CreateSceneProxy() override;
+	virtual FBoxSphereBounds CalcBounds(const FTransform& LocalToWorld) const;
+
+	const FContextualAnimData* FindBestDataForActor(const AActor* Querier) const;
+
+	const FContextualAnimData* FindClosestDataForActor(const AActor* Querier) const;
+
+	UFUNCTION(BlueprintCallable, Category = "Contextual Animation")
+	bool FindBestEntryPointForActor(const AActor* Querier, FContextualAnimEntryPoint& Result) const;
+
+	UFUNCTION(BlueprintCallable, Category = "Contextual Animation")
+	bool FindClosestEntryPointForActor(const AActor* Querier, FContextualAnimEntryPoint& Result) const;
+};
diff --git a/Engine/Plugins/Experimental/Animation/ContextualAnimation/Source/ContextualAnimation/Public/ContextualAnimation.h b/Engine/Plugins/Experimental/Animation/ContextualAnimation/Source/ContextualAnimation/Public/ContextualAnimation.h
new file mode 100644
--- /dev/null
+++ b/Engine/Plugins/Experimental/Animation/ContextualAnimation/Source/ContextualAnimation/Public/ContextualAnimation.h
@@ -0,0 +1,15 @@
+// Copyright Epic Games, Inc. All Rights Reserved.
+
+#pragma once
+
+#include "CoreMinimal.h"
+#include "Modules/ModuleManager.h"
+
+class FContextualAnimationModule : public IModuleInterface
+{
+public:
+
+	/** IModuleInterface implementation */
+	virtual void StartupModule() override;
+	virtual void ShutdownModule() override;
+};
diff --git a/Engine/Plugins/Experimental/Animation/ContextualAnimation/Source/ContextualAnimationEditor/ContextualAnimationEditor.Build.cs b/Engine/Plugins/Experimental/Animation/ContextualAnimation/Source/ContextualAnimationEditor/ContextualAnimationEditor.Build.cs
new file mode 100644
--- /dev/null
+++ b/Engine/Plugins/Experimental/Animation/ContextualAnimation/Source/ContextualAnimationEditor/ContextualAnimationEditor.Build.cs
@@ -0,0 +1,62 @@
+// Copyright Epic Games, Inc. All Rights Reserved.
+
+using UnrealBuildTool;
+
+public class ContextualAnimationEditor : ModuleRules
+{
+	public ContextualAnimationEditor(ReadOnlyTargetRules Target) : base(Target)
+	{
+		PCHUsage = ModuleRules.PCHUsageMode.UseExplicitOrSharedPCHs;
+		
+		PublicIncludePaths.AddRange(
+			new string[] {
+				// ... add public include paths required here ...
+			}
+			);
+				
+		
+		PrivateIncludePaths.AddRange(
+			new string[] {
+
+			}
+			);
+			
+		
+		PublicDependencyModuleNames.AddRange(
+			new string[]
+			{
+				"Core",
+				"AIModule",
+				"NavigationSystem",
+				"ContextualAnimation",
+				"MotionWarping",
+			}
+			);
+			
+		
+		PrivateDependencyModuleNames.AddRange(
+			new string[]
+			{
+				"CoreUObject",
+				"Engine",
+				"Slate",
+				"SlateCore",
+				"InputCore",
+				"EditorFramework",
+				"UnrealEd",
+				"LevelEditor",
+				"DetailCustomizations",
+				"EditorStyle",
+				"Projects"
+			}
+			);
+		
+		
+		DynamicallyLoadedModuleNames.AddRange(
+			new string[]
+			{
+				// ... add any modules that your module loads dynamically here ...
+			}
+			);
+	}
+}
diff --git a/Engine/Plugins/Experimental/Animation/ContextualAnimation/Source/ContextualAnimationEditor/Private/ContextualAnimEdMode.cpp b/Engine/Plugins/Experimental/Animation/ContextualAnimation/Source/ContextualAnimationEditor/Private/ContextualAnimEdMode.cpp
new file mode 100644
--- /dev/null
+++ b/Engine/Plugins/Experimental/Animation/ContextualAnimation/Source/ContextualAnimationEditor/Private/ContextualAnimEdMode.cpp
@@ -0,0 +1,201 @@
+// Copyright Epic Games, Inc. All Rights Reserved.
+
+#include "ContextualAnimEdMode.h"
+#include "ContextualAnimEdModeToolkit.h"
+#include "Toolkits/ToolkitManager.h"
+#include "EditorModeManager.h"
+#include "Engine/Selection.h"
+#include "EngineUtils.h"
+#include "ContextualAnimEdModeSettings.h"
+#include "DrawDebugHelpers.h"
+#include "GameFramework/Character.h"
+#include "NavigationSystem.h"
+#include "AIController.h"
+#include "ContextualAnimComponent.h"
+#include "MotionWarpingComponent.h"
+#include "Animation/AnimMontage.h"
+
+const FEditorModeID FContextualAnimEdMode::EM_ContextualAnimEdModeId = TEXT("EM_ContextualAnimEdMode");
+
+FContextualAnimEdMode::FContextualAnimEdMode()
+{
+}
+
+FContextualAnimEdMode::~FContextualAnimEdMode()
+{
+}
+
+void FContextualAnimEdMode::Enter()
+{
+	FEdMode::Enter();
+
+	if (!Toolkit.IsValid() && UsesToolkits())
+	{
+		Toolkit = MakeShareable(new FContextualAnimEdModeToolkit);
+		Toolkit->Init(Owner->GetToolkitHost());
+	}
+}
+
+void FContextualAnimEdMode::Exit()
+{
+	if (Toolkit.IsValid())
+	{
+		FToolkitManager::Get().CloseToolkit(Toolkit.ToSharedRef());
+		Toolkit.Reset();
+	}
+
+	FEdMode::Exit();
+}
+
+TSharedPtr<FContextualAnimEdModeToolkit> FContextualAnimEdMode::GetContextualAnimEdModeToolkit()
+{
+	return StaticCastSharedPtr<FContextualAnimEdModeToolkit>(Toolkit);
+}
+
+void FContextualAnimEdMode::Tick(FEditorViewportClient* ViewportClient, float DeltaTime)
+{
+	FEdMode::Tick(ViewportClient, DeltaTime);
+
+	if(TestCharacter.IsValid() && GEditor->IsSimulatingInEditor())
+	{
+		if (ViewportClient->Viewport->KeyState(EKeys::W) || ViewportClient->Viewport->KeyState(EKeys::S))
+		{
+			const FVector WorldDirection = FRotationMatrix(TestCharacter->GetActorRotation()).GetScaledAxis(EAxis::X);
+			TestCharacter->AddMovementInput(WorldDirection, ViewportClient->Viewport->KeyState(EKeys::W) ? 1.f : -1.f);
+		}
+
+		if (ViewportClient->Viewport->KeyState(EKeys::A) || ViewportClient->Viewport->KeyState(EKeys::D))
+		{
+			const FVector WorldDirection = FRotationMatrix(TestCharacter->GetActorRotation()).GetScaledAxis(EAxis::Y);
+			TestCharacter->AddMovementInput(WorldDirection, ViewportClient->Viewport->KeyState(EKeys::D) ? 1.f : -1.f);
+		}
+	}
+}
+
+bool FContextualAnimEdMode::InputKey(FEditorViewportClient* ViewportClient, FViewport* Viewport, FKey Key, EInputEvent Event)
+{
+	if(Key == EKeys::Enter && Event == IE_Released && TestCharacter.IsValid() && GEditor->IsSimulatingInEditor())
+	{
+		if (!TestCharacter->IsPlayingRootMotion())
+		{
+			//@TODO: Cache the list of interactable actors
+			float BestDistance = MAX_FLT;
+			UContextualAnimComponent* ContextualAnimComp = nullptr;
+			for (TActorIterator<AActor> It(GetWorld()); It; ++It)
+			{
+				AActor* Actor = *It;
+				UContextualAnimComponent* Comp = Actor ? Actor->FindComponentByClass<UContextualAnimComponent>() : nullptr;
+				if(Comp)
+				{
+					const float Distance = FVector::DistSquared(Actor->GetActorLocation(), TestCharacter->GetActorLocation());
+					if(Distance < BestDistance)
+					{
+						BestDistance = Distance;
+						ContextualAnimComp = Comp;
+					}
+				}
+			}
+
+			if (ContextualAnimComp)
+			{
+				FContextualAnimEntryPoint EntryPoint;
+				if (ContextualAnimComp->FindBestEntryPointForActor(TestCharacter.Get(), EntryPoint))
+				{
+					if (UAnimMontage* Montage = EntryPoint.Animation.LoadSynchronous())
+					{
+						if (UMotionWarpingComponent* MotionWarpingComp = TestCharacter->FindComponentByClass<UMotionWarpingComponent>())
+						{
+							const FName SyncPointName = GetContextualAnimEdModeToolkit()->GetSettings()->MotionWarpSyncPointName;
+							MotionWarpingComp->AddOrUpdateSyncPoint(SyncPointName, EntryPoint.SyncTransform);
+							TestCharacter->PlayAnimMontage(Montage);
+						}
+					}
+				}
+				else
+				{
+					GEngine->AddOnScreenDebugMessage(-1, 5.f, FColor::Yellow, FString::Printf(TEXT("Found nearest interactable actor (%s) but the Test Character is not close enough to any of the entry points"), 
+					*GetNameSafe(ContextualAnimComp->GetOwner())));
+				}
+			}
+			else
+			{
+				GEngine->AddOnScreenDebugMessage(-1, 5.f, FColor::Yellow, FString(TEXT("There are no interactable actors in the level")));
+			}
+		}
+		else
+		{
+			TestCharacter->StopAnimMontage();
+		}
+
+		return true;
+	}
+
+	return FEdMode::InputKey(ViewportClient, Viewport, Key, Event);
+}
+
+bool FContextualAnimEdMode::HandleClick(FEditorViewportClient* InViewportClient, HHitProxy* HitProxy, const FViewportClick& Click)
+{
+	if (Click.IsControlDown() && GEditor->IsSimulatingInEditor())
+	{
+		FSceneViewFamilyContext ViewFamily(FSceneViewFamily::ConstructionValues(InViewportClient->Viewport, InViewportClient->GetScene(), InViewportClient->EngineShowFlags).SetRealtimeUpdate(InViewportClient->IsRealtime()));
+		FSceneView* View = InViewportClient->CalcSceneView(&ViewFamily);
+		FViewportCursorLocation Cursor(View, InViewportClient, Click.GetClickPos().X, Click.GetClickPos().Y);
+		const auto ViewportType = InViewportClient->GetViewportType();
+		
+		const FVector RayStart = Cursor.GetOrigin();
+		const FVector RayEnd = RayStart + Cursor.GetDirection() * HALF_WORLD_MAX;
+				
+		FHitResult HitResult;
+		InViewportClient->GetWorld()->LineTraceSingleByChannel(HitResult, RayStart, RayEnd, ECC_WorldStatic, FCollisionQueryParams::DefaultQueryParam);
+		
+		if (HitResult.IsValidBlockingHit())
+		{
+			if(!TestCharacter.IsValid())
+			{
+				TSubclassOf<ACharacter> TestActorClass = GetContextualAnimEdModeToolkit()->GetSettings()->TestCharacterClass;
+				if(TestActorClass)
+				{
+					TestCharacter = GetWorld()->SpawnActor<ACharacter>(TestActorClass, FTransform(HitResult.ImpactPoint));
+					
+					if(TestCharacter.IsValid())
+					{
+						if (!TestCharacter->AIControllerClass || !TestCharacter->AIControllerClass->IsChildOf<AAIController>())
+						{
+							TestCharacter->AIControllerClass = AAIController::StaticClass();
+						}
+
+						TestCharacter->SpawnDefaultController();
+					}
+				}
+				else
+				{
+					GEngine->AddOnScreenDebugMessage(-1, 5.f, FColor::Yellow, FString(TEXT("Select a valid character class for your test actor")));
+				}
+			}
+			else
+			{
+				if(AAIController* Controller = Cast<AAIController>(TestCharacter->GetController()))
+				{
+					UNavigationSystemV1* NavSys = FNavigationSystem::GetCurrent<UNavigationSystemV1>(Controller->GetWorld());
+					const ANavigationData* NavData = NavSys ? NavSys->GetNavDataForProps(Controller->GetNavAgentPropertiesRef(), Controller->GetNavAgentLocation()) : nullptr;
+
+					const bool bUsePathfinding = (NavData != nullptr);
+					Controller->MoveToLocation(HitResult.ImpactPoint, 10.f, true, bUsePathfinding);
+				}
+			}
+		}
+
+		return true;
+	}
+
+	return FEdMode::HandleClick(InViewportClient, HitProxy, Click);
+}
+
+bool FContextualAnimEdMode::UsesToolkits() const
+{
+	return true;
+}
+
+
+
+
diff --git a/Engine/Plugins/Experimental/Animation/ContextualAnimation/Source/ContextualAnimationEditor/Private/ContextualAnimEdModeSettings.cpp b/Engine/Plugins/Experimental/Animation/ContextualAnimation/Source/ContextualAnimationEditor/Private/ContextualAnimEdModeSettings.cpp
new file mode 100644
--- /dev/null
+++ b/Engine/Plugins/Experimental/Animation/ContextualAnimation/Source/ContextualAnimationEditor/Private/ContextualAnimEdModeSettings.cpp
@@ -0,0 +1,9 @@
+// Copyright Epic Games, Inc. All Rights Reserved.
+
+#include "ContextualAnimEdModeSettings.h"
+
+UContextualAnimEdModeSettings::UContextualAnimEdModeSettings(const FObjectInitializer& ObjectInitializer)
+	: Super(ObjectInitializer)
+{
+	MotionWarpSyncPointName = FName(TEXT("Interact"));
+}
\ No newline at end of file
diff --git a/Engine/Plugins/Experimental/Animation/ContextualAnimation/Source/ContextualAnimationEditor/Private/ContextualAnimEdModeToolkit.cpp b/Engine/Plugins/Experimental/Animation/ContextualAnimation/Source/ContextualAnimationEditor/Private/ContextualAnimEdModeToolkit.cpp
new file mode 100644
--- /dev/null
+++ b/Engine/Plugins/Experimental/Animation/ContextualAnimation/Source/ContextualAnimationEditor/Private/ContextualAnimEdModeToolkit.cpp
@@ -0,0 +1,87 @@
+// Copyright Epic Games, Inc. All Rights Reserved.
+
+#include "ContextualAnimEdModeToolkit.h"
+#include "ContextualAnimEdMode.h"
+#include "Engine/Selection.h"
+#include "Widgets/Input/SButton.h"
+#include "Widgets/Text/STextBlock.h"
+#include "EditorModeManager.h"
+#include "ContextualAnimEdModeSettings.h"
+
+#define LOCTEXT_NAMESPACE "ContextualAnimEdModeToolkit"
+
+FContextualAnimEdModeToolkit::FContextualAnimEdModeToolkit()
+	: Settings(nullptr)
+{
+	Settings = NewObject<UContextualAnimEdModeSettings>(UContextualAnimEdModeSettings::StaticClass());
+}
+
+void FContextualAnimEdModeToolkit::AddReferencedObjects(FReferenceCollector& Collector)
+{
+	Collector.AddReferencedObject(Settings);
+}
+
+void FContextualAnimEdModeToolkit::Init(const TSharedPtr<IToolkitHost>& InitToolkitHost)
+{
+	FPropertyEditorModule& PropertyModule = FModuleManager::LoadModuleChecked<FPropertyEditorModule>("PropertyEditor");
+
+	FDetailsViewArgs Args;
+	Args.bAllowSearch = false;
+	Args.bHideSelectionTip = true;
+
+	EdModeSettingsWidget = PropertyModule.CreateDetailView(Args);
+	EdModeSettingsWidget->SetObject(Settings);
+
+	SAssignNew(ToolkitWidget, SBorder)
+		.HAlign(HAlign_Center)
+		.Padding(5)
+		[
+			SNew(SVerticalBox)
+			
+			+ SVerticalBox::Slot()
+			.HAlign(HAlign_Left)
+			.AutoHeight()
+			.Padding(5)
+			[
+				SNew(STextBlock)
+				.AutoWrapText(true)
+				.Text(FText::FromString("\
+					- Select the class for your test actor\n\
+					- Start Simulating Mode\n\
+					- Ctrl + Click to spawn a test actor\n\
+					- Ctrl + Click or WASD to move the test actor around\n\
+					- [Enter] to start / stop an interaction"))
+			]
+
+			+ SVerticalBox::Slot()
+			.HAlign(HAlign_Left)
+			.AutoHeight()
+			[
+				EdModeSettingsWidget.ToSharedRef()
+			]
+		];
+
+	FModeToolkit::Init(InitToolkitHost);
+}
+
+FName FContextualAnimEdModeToolkit::GetToolkitFName() const
+{
+	return FName("ContextualAnimEdMode");
+}
+
+FText FContextualAnimEdModeToolkit::GetBaseToolkitName() const
+{
+	return NSLOCTEXT("ContextualAnimEdModeToolkit", "DisplayName", "Contextual Anim Tool");
+}
+
+FEdMode* FContextualAnimEdModeToolkit::GetEditorMode() const
+{
+	return GLevelEditorModeTools().GetActiveMode(FContextualAnimEdMode::EM_ContextualAnimEdModeId);
+}
+
+FContextualAnimEdMode* FContextualAnimEdModeToolkit::GetContextualAnimEdMode() const
+{
+	return static_cast<FContextualAnimEdMode*>(GetEditorMode());
+}
+
+#undef LOCTEXT_NAMESPACE
diff --git a/Engine/Plugins/Experimental/Animation/ContextualAnimation/Source/ContextualAnimationEditor/Private/ContextualAnimEditorStyle.cpp b/Engine/Plugins/Experimental/Animation/ContextualAnimation/Source/ContextualAnimationEditor/Private/ContextualAnimEditorStyle.cpp
new file mode 100644
--- /dev/null
+++ b/Engine/Plugins/Experimental/Animation/ContextualAnimation/Source/ContextualAnimationEditor/Private/ContextualAnimEditorStyle.cpp
@@ -0,0 +1,70 @@
+// Copyright Epic Games, Inc. All Rights Reserved.
+
+#include "ContextualAnimEditorStyle.h"
+#include "Styling/SlateStyleRegistry.h"
+#include "Framework/Application/SlateApplication.h"
+#include "Slate/SlateGameResources.h"
+#include "Interfaces/IPluginManager.h"
+
+TSharedPtr< FSlateStyleSet > FContextualAnimEditorStyle::StyleInstance = NULL;
+
+void FContextualAnimEditorStyle::Initialize()
+{
+	if (!StyleInstance.IsValid())
+	{
+		StyleInstance = Create();
+		FSlateStyleRegistry::RegisterSlateStyle(*StyleInstance);
+	}
+}
+
+void FContextualAnimEditorStyle::Shutdown()
+{
+	FSlateStyleRegistry::UnRegisterSlateStyle(*StyleInstance);
+	ensure(StyleInstance.IsUnique());
+	StyleInstance.Reset();
+}
+
+FName FContextualAnimEditorStyle::GetStyleSetName()
+{
+	static FName StyleSetName(TEXT("ContextualAnimEditorStyle"));
+	return StyleSetName;
+}
+
+#define IMAGE_BRUSH( RelativePath, ... ) FSlateImageBrush( Style->RootToContentDir( RelativePath, TEXT(".png") ), __VA_ARGS__ )
+#define BOX_BRUSH( RelativePath, ... ) FSlateBoxBrush( Style->RootToContentDir( RelativePath, TEXT(".png") ), __VA_ARGS__ )
+#define BORDER_BRUSH( RelativePath, ... ) FSlateBorderBrush( Style->RootToContentDir( RelativePath, TEXT(".png") ), __VA_ARGS__ )
+#define TTF_FONT( RelativePath, ... ) FSlateFontInfo( Style->RootToContentDir( RelativePath, TEXT(".ttf") ), __VA_ARGS__ )
+#define OTF_FONT( RelativePath, ... ) FSlateFontInfo( Style->RootToContentDir( RelativePath, TEXT(".otf") ), __VA_ARGS__ )
+
+const FVector2D Icon16x16(16.0f, 16.0f);
+const FVector2D Icon20x20(20.0f, 20.0f);
+const FVector2D Icon40x40(40.0f, 40.0f);
+
+TSharedRef< FSlateStyleSet > FContextualAnimEditorStyle::Create()
+{
+	TSharedRef< FSlateStyleSet > Style = MakeShareable(new FSlateStyleSet("ContextualAnimEditorStyle"));
+	Style->SetContentRoot(IPluginManager::Get().FindPlugin("ContextualAnimation")->GetBaseDir() / TEXT("Content"));
+
+	Style->Set("ContextualAnimEditor.Icon", new IMAGE_BRUSH(TEXT("ButtonIcon_40x"), Icon40x40));
+
+	return Style;
+}
+
+#undef IMAGE_BRUSH
+#undef BOX_BRUSH
+#undef BORDER_BRUSH
+#undef TTF_FONT
+#undef OTF_FONT
+
+void FContextualAnimEditorStyle::ReloadTextures()
+{
+	if (FSlateApplication::IsInitialized())
+	{
+		FSlateApplication::Get().GetRenderer()->ReloadTextureResources();
+	}
+}
+
+const ISlateStyle& FContextualAnimEditorStyle::Get()
+{
+	return *StyleInstance;
+}
diff --git a/Engine/Plugins/Experimental/Animation/ContextualAnimation/Source/ContextualAnimationEditor/Private/ContextualAnimationEditor.cpp b/Engine/Plugins/Experimental/Animation/ContextualAnimation/Source/ContextualAnimationEditor/Private/ContextualAnimationEditor.cpp
new file mode 100644
--- /dev/null
+++ b/Engine/Plugins/Experimental/Animation/ContextualAnimation/Source/ContextualAnimationEditor/Private/ContextualAnimationEditor.cpp
@@ -0,0 +1,42 @@
+// Copyright Epic Games, Inc. All Rights Reserved.
+
+#include "ContextualAnimationEditor.h"
+#include "ContextualAnimEdMode.h"
+#include "ContextualAnimEditorStyle.h"
+#include "Customizations/MathStructCustomizations.h"
+
+#define LOCTEXT_NAMESPACE "FContextualAnimationEditorModule"
+
+void FContextualAnimationEditorModule::StartupModule()
+{
+	FContextualAnimEditorStyle::Initialize();
+	FContextualAnimEditorStyle::ReloadTextures();
+
+	FPropertyEditorModule& PropertyModule = FModuleManager::LoadModuleChecked<FPropertyEditorModule>("PropertyEditor");
+	PropertyModule.RegisterCustomPropertyTypeLayout("ContextualAnimDistanceParam", FOnGetPropertyTypeCustomizationInstance::CreateStatic(&FMathStructCustomization::MakeInstance));
+	PropertyModule.RegisterCustomPropertyTypeLayout("ContextualAnimAngleParam", FOnGetPropertyTypeCustomizationInstance::CreateStatic(&FMathStructCustomization::MakeInstance));
+
+	FEditorModeRegistry::Get().RegisterMode<FContextualAnimEdMode>(FContextualAnimEdMode::EM_ContextualAnimEdModeId, 
+		LOCTEXT("ContextualAnimEdModeEdModeName", "ContextualAnim"), 
+		FSlateIcon(FContextualAnimEditorStyle::GetStyleSetName(), "ContextualAnimEditor.Icon", "ContextualAnimEditor.Icon"),
+		true);
+}
+
+void FContextualAnimationEditorModule::ShutdownModule()
+{
+	FContextualAnimEditorStyle::Shutdown();
+
+	FPropertyEditorModule& PropertyModule = FModuleManager::LoadModuleChecked<FPropertyEditorModule>("PropertyEditor");
+	PropertyModule.UnregisterCustomPropertyTypeLayout("ContextualAnimFloatParam");
+
+	FEditorModeRegistry::Get().UnregisterMode(FContextualAnimEdMode::EM_ContextualAnimEdModeId);
+}
+
+FContextualAnimationEditorModule& FContextualAnimationEditorModule::Get()
+{
+	return FModuleManager::Get().GetModuleChecked<FContextualAnimationEditorModule>("ContextualAnimationEditor");
+}
+
+#undef LOCTEXT_NAMESPACE
+	
+IMPLEMENT_MODULE(FContextualAnimationEditorModule, ContextualAnimationEditor)
\ No newline at end of file
diff --git a/Engine/Plugins/Experimental/Animation/ContextualAnimation/Source/ContextualAnimationEditor/Public/ContextualAnimEdMode.h b/Engine/Plugins/Experimental/Animation/ContextualAnimation/Source/ContextualAnimationEditor/Public/ContextualAnimEdMode.h
new file mode 100644
--- /dev/null
+++ b/Engine/Plugins/Experimental/Animation/ContextualAnimation/Source/ContextualAnimationEditor/Public/ContextualAnimEdMode.h
@@ -0,0 +1,30 @@
+// Copyright Epic Games, Inc. All Rights Reserved.
+
+#pragma once
+
+#include "CoreMinimal.h"
+#include "EdMode.h"
+
+class FContextualAnimEdModeToolkit;
+
+class FContextualAnimEdMode : public FEdMode
+{
+public:
+	const static FEditorModeID EM_ContextualAnimEdModeId;
+public:
+	FContextualAnimEdMode();
+	virtual ~FContextualAnimEdMode();
+
+	// FEdMode interface
+	virtual void Enter() override;
+	virtual void Exit() override;
+	virtual void Tick(FEditorViewportClient* ViewportClient, float DeltaTime) override;
+	virtual bool InputKey(FEditorViewportClient* ViewportClient, FViewport* Viewport, FKey Key, EInputEvent Event) override;
+	virtual bool HandleClick(FEditorViewportClient* InViewportClient, HHitProxy* HitProxy, const FViewportClick& Click) override;
+	virtual bool UsesToolkits() const override;
+	// End of FEdMode interface
+
+	TSharedPtr<FContextualAnimEdModeToolkit> GetContextualAnimEdModeToolkit();
+
+	TWeakObjectPtr<class ACharacter> TestCharacter;
+};
diff --git a/Engine/Plugins/Experimental/Animation/ContextualAnimation/Source/ContextualAnimationEditor/Public/ContextualAnimEdModeSettings.h b/Engine/Plugins/Experimental/Animation/ContextualAnimation/Source/ContextualAnimationEditor/Public/ContextualAnimEdModeSettings.h
new file mode 100644
--- /dev/null
+++ b/Engine/Plugins/Experimental/Animation/ContextualAnimation/Source/ContextualAnimationEditor/Public/ContextualAnimEdModeSettings.h
@@ -0,0 +1,24 @@
+// Copyright Epic Games, Inc. All Rights Reserved.
+
+#pragma once
+
+#include "CoreMinimal.h"
+#include "ContextualAnimEdModeSettings.generated.h"
+
+class ACharacter;
+
+UCLASS()
+class UContextualAnimEdModeSettings : public UObject
+{
+	GENERATED_BODY()
+
+public:
+
+	UPROPERTY(EditAnywhere, Category = "Settings")
+	TSubclassOf<ACharacter> TestCharacterClass;
+
+	UPROPERTY(EditAnywhere, Category = "Settings")
+	FName MotionWarpSyncPointName;
+
+	UContextualAnimEdModeSettings(const FObjectInitializer& ObjectInitializer);
+};
\ No newline at end of file
diff --git a/Engine/Plugins/Experimental/Animation/ContextualAnimation/Source/ContextualAnimationEditor/Public/ContextualAnimEdModeToolkit.h b/Engine/Plugins/Experimental/Animation/ContextualAnimation/Source/ContextualAnimationEditor/Public/ContextualAnimEdModeToolkit.h
new file mode 100644
--- /dev/null
+++ b/Engine/Plugins/Experimental/Animation/ContextualAnimation/Source/ContextualAnimationEditor/Public/ContextualAnimEdModeToolkit.h
@@ -0,0 +1,39 @@
+// Copyright Epic Games, Inc. All Rights Reserved.
+
+#pragma once
+
+#include "CoreMinimal.h"
+#include "Toolkits/BaseToolkit.h"
+
+class UContextualAnimEdModeSettings;
+class FContextualAnimEdMode;
+
+class FContextualAnimEdModeToolkit : public FModeToolkit, public FGCObject
+{
+public:
+
+	FContextualAnimEdModeToolkit();
+	
+	/** FModeToolkit interface */
+	virtual void Init(const TSharedPtr<IToolkitHost>& InitToolkitHost) override;
+
+	/** IToolkit interface */
+	virtual FName GetToolkitFName() const override;
+	virtual FText GetBaseToolkitName() const override;
+	virtual class FEdMode* GetEditorMode() const override;
+	virtual TSharedPtr<class SWidget> GetInlineContent() const override { return ToolkitWidget; }
+
+	/** FGCObject interface */
+	virtual void AddReferencedObjects(FReferenceCollector& Collector) override;
+
+	FContextualAnimEdMode* GetContextualAnimEdMode() const;
+
+	UContextualAnimEdModeSettings* GetSettings() const { return Settings; }
+
+private:
+
+	UContextualAnimEdModeSettings* Settings;
+
+	TSharedPtr<SWidget> ToolkitWidget;
+	TSharedPtr<IDetailsView> EdModeSettingsWidget;
+};
diff --git a/Engine/Plugins/Experimental/Animation/ContextualAnimation/Source/ContextualAnimationEditor/Public/ContextualAnimEditorStyle.h b/Engine/Plugins/Experimental/Animation/ContextualAnimation/Source/ContextualAnimationEditor/Public/ContextualAnimEditorStyle.h
new file mode 100644
--- /dev/null
+++ b/Engine/Plugins/Experimental/Animation/ContextualAnimation/Source/ContextualAnimationEditor/Public/ContextualAnimEditorStyle.h
@@ -0,0 +1,29 @@
+// Copyright Epic Games, Inc. All Rights Reserved.
+
+#pragma once
+
+#include "CoreMinimal.h"
+#include "Styling/SlateStyle.h"
+
+class FContextualAnimEditorStyle
+{
+public:
+
+	static void Initialize();
+
+	static void Shutdown();
+
+	/** reloads textures used by slate renderer */
+	static void ReloadTextures();
+
+	/** @return The Slate style set for the Shooter game */
+	static const ISlateStyle& Get();
+
+	static FName GetStyleSetName();
+
+private:
+
+	static TSharedRef< class FSlateStyleSet > Create();
+
+	static TSharedPtr< class FSlateStyleSet > StyleInstance;
+};
\ No newline at end of file
diff --git a/Engine/Plugins/Experimental/Animation/ContextualAnimation/Source/ContextualAnimationEditor/Public/ContextualAnimationEditor.h b/Engine/Plugins/Experimental/Animation/ContextualAnimation/Source/ContextualAnimationEditor/Public/ContextualAnimationEditor.h
new file mode 100644
--- /dev/null
+++ b/Engine/Plugins/Experimental/Animation/ContextualAnimation/Source/ContextualAnimationEditor/Public/ContextualAnimationEditor.h
@@ -0,0 +1,17 @@
+// Copyright Epic Games, Inc. All Rights Reserved.
+
+#pragma once
+
+#include "CoreMinimal.h"
+#include "Modules/ModuleManager.h"
+
+class FContextualAnimationEditorModule : public IModuleInterface
+{
+public:
+
+	/** IModuleInterface implementation */
+	virtual void StartupModule() override;
+	virtual void ShutdownModule() override;
+
+	static FContextualAnimationEditorModule& Get();
+};
