Contextual Anim - Added function to find best start time for the animation based on distance from the object. ContextualAnimComp now extend from SphereComp. Exposed sample rate for alignment extraction as param. Refactored API to query data #review-15211083 [REDACTED_HANDLE] [REDACTED_HANDLE] #fyi [REDACTED_HANDLE]

diff --git a/Engine/Plugins/Experimental/Animation/ContextualAnimation/Source/ContextualAnimation/Private/ContextualAnimAsset.cpp b/Engine/Plugins/Experimental/Animation/ContextualAnimation/Source/ContextualAnimation/Private/ContextualAnimAsset.cpp
--- a/Engine/Plugins/Experimental/Animation/ContextualAnimation/Source/ContextualAnimation/Private/ContextualAnimAsset.cpp
+++ b/Engine/Plugins/Experimental/Animation/ContextualAnimation/Source/ContextualAnimation/Private/ContextualAnimAsset.cpp
@@ -65,27 +65,57 @@ FTransform FAlignmentTrackContainer::ExtractTransformAtTime(float Time) const
 {
 	FTransform AlignmentTransform = FTransform::Identity;
 	const int32 TotalFrames = Track.PosKeys.Num();
-	const float TrackLength = (TotalFrames - 1) * SampleRate;
+	const float TrackLength = (TotalFrames - 1) * SampleInterval;
 	FAnimationUtils::ExtractTransformFromTrack(Time, TotalFrames, TrackLength, Track, EAnimInterpolationType::Linear, AlignmentTransform);
 	return AlignmentTransform;
 }
 
+// FContextualAnimData
+///////////////////////////////////////////////////////////////////////
+
+float FContextualAnimData::FindBestAnimStartTime(const FVector& LocalLocation) const
+{
+	float BestTime = 0.f;
+
+	const FTransform TransformAtEntryTime = GetAlignmentTransformAtEntryTime();
+	const FVector Direction = (LocalLocation - TransformAtEntryTime.GetLocation()).GetSafeNormal2D();
+	const float Dot = FVector::DotProduct(Direction, TransformAtEntryTime.GetRotation().GetForwardVector());
+	if (Dot > 0.f)
+	{
+		float BestDistance = MAX_FLT;	
+		const TArray<FVector>& PosKeys = AlignmentData.Track.PosKeys;
+
+		//@TODO: Very simple search for now
+		for (int32 Idx = 0; Idx < PosKeys.Num(); Idx++)
+		{
+			const float Time = Idx * AlignmentData.SampleInterval;
+			if(AnimMinStartTime > 0.f && Time >= AnimMinStartTime)
+			{
+				break;
+			}
+
+			const float DistSquared2D = FVector::DistSquared2D(LocalLocation, PosKeys[Idx]);
+			if (DistSquared2D < BestDistance)
+			{
+				BestDistance = DistSquared2D;
+				BestTime = Time;
+			}
+		}
+	}
+
+	return BestTime;
+}
+
 // UContextualAnimAsset
 ///////////////////////////////////////////////////////////////////////
 
 UContextualAnimAsset::UContextualAnimAsset(const FObjectInitializer& ObjectInitializer)
 	: Super(ObjectInitializer) 
 {
+	SampleRate = 15;
 	MeshToComponent = FTransform(FRotator(0.f, -90.f, 0.f));
 }
 
-#if WITH_EDITOR
-void UContextualAnimAsset::PostEditChangeProperty(FPropertyChangedEvent& PropertyChangedEvent)
-{
-	Super::PostEditChangeProperty(PropertyChangedEvent);
-}
-#endif // WITH_EDITOR
-
 void UContextualAnimAsset::PreSave(const class ITargetPlatform* TargetPlatform)
 {
 	const FTransform MeshToComponentInverse = MeshToComponent.Inverse();
@@ -94,6 +124,8 @@ void UContextualAnimAsset::PreSave(const class ITargetPlatform* TargetPlatform)
 	RequiredBoneIndexArray.Add(0);
 	FBoneContainer BoneContainer;
 
+	const float SampleInterval = 1.f / SampleRate;
+
 	for (FContextualAnimData& Data : DataContainer)
 	{
 		const UAnimMontage* Animation = Data.Animation.LoadSynchronous();
@@ -107,13 +139,12 @@ void UContextualAnimAsset::PreSave(const class ITargetPlatform* TargetPlatform)
 		float SectionStartTime, SectionEndTime;
 		Animation->GetSectionStartAndEndTime(0, SectionStartTime, SectionEndTime);
 
-		const float SampleRate = 1 / 30.f;
 		float Time = 0.f;
 		float EndTime = SectionEndTime;
 		int32 SampleIndex = 0;
 
 		Data.SyncTime = EndTime;
-		Data.AlignmentData.Initialize(SampleRate);
+		Data.AlignmentData.Initialize(SampleInterval);
 		FRawAnimSequenceTrack& NewTrack = Data.AlignmentData.Track;
 
 		if (AlignmentJoint.IsNone())
@@ -125,7 +156,7 @@ void UContextualAnimAsset::PreSave(const class ITargetPlatform* TargetPlatform)
 
 			while (Time < EndTime)
 			{
-				Time = FMath::Clamp(SampleIndex * SampleRate, 0.f, EndTime);
+				Time = FMath::Clamp(SampleIndex * SampleInterval, 0.f, EndTime);
 				SampleIndex++;
 
 				FCSPose<FCompactPose> ComponentSpacePose;
@@ -160,7 +191,7 @@ void UContextualAnimAsset::PreSave(const class ITargetPlatform* TargetPlatform)
 
 			while (Time < EndTime)
 			{
-				Time = FMath::Clamp(SampleIndex * SampleRate, 0.f, EndTime);
+				Time = FMath::Clamp(SampleIndex * SampleInterval, 0.f, EndTime);
 				SampleIndex++;
 
 				FCSPose<FCompactPose> ComponentSpacePose;
diff --git a/Engine/Plugins/Experimental/Animation/ContextualAnimation/Source/ContextualAnimation/Private/ContextualAnimComponent.cpp b/Engine/Plugins/Experimental/Animation/ContextualAnimation/Source/ContextualAnimation/Private/ContextualAnimComponent.cpp
--- a/Engine/Plugins/Experimental/Animation/ContextualAnimation/Source/ContextualAnimation/Private/ContextualAnimComponent.cpp
+++ b/Engine/Plugins/Experimental/Animation/ContextualAnimation/Source/ContextualAnimation/Private/ContextualAnimComponent.cpp
@@ -9,132 +9,126 @@ UContextualAnimComponent::UContextualAnimComponent(const FObjectInitializer& Obj
 {
 	PrimaryComponentTick.bCanEverTick = true;
 	PrimaryComponentTick.bStartWithTickEnabled = false;
-}
 
-void UContextualAnimComponent::TickComponent(float DeltaTime, enum ELevelTick TickType, FActorComponentTickFunction* ThisTickFunction)
-{
-	Super::TickComponent(DeltaTime,TickType, ThisTickFunction);
+	SetCollisionEnabled(ECollisionEnabled::QueryOnly);
+	SetCollisionResponseToAllChannels(ECR_Ignore);
+	SetCollisionResponseToChannel(ECC_Pawn, ECR_Overlap);
+	SetGenerateOverlapEvents(true);
+
+	bHiddenInGame = false;
 }
 
-const FContextualAnimData* UContextualAnimComponent::FindBestDataForActor(const AActor* Querier) const
+bool UContextualAnimComponent::QueryData(const FContextualAnimQueryParams& QueryParams, FContextualAnimQueryResult& Result) const
 {
-	if(ContextualAnimAsset == nullptr || Querier == nullptr)
+	Result.Reset();
+
+	if (!ContextualAnimAsset || !QueryParams.Querier.IsValid())
 	{
-		return nullptr;
+		return false;
 	}
 
-	for (int32 Idx = 0; Idx < ContextualAnimAsset->DataContainer.Num(); Idx++)
+	int32 DataIndex = INDEX_NONE;
+	if (QueryParams.bComplexQuery)
 	{
-		const FContextualAnimData& Data = ContextualAnimAsset->DataContainer[Idx];
-		const FTransform EntryTransform = Data.GetAlignmentTransformAtEntryTime() * GetComponentTransform();
-		
-		FVector Origin = GetComponentTransform().GetLocation();
-		FVector Direction = (EntryTransform.GetLocation() - Origin).GetSafeNormal2D();
-
-		if (Data.OffsetFromOrigin != 0.f)
+		for (int32 Idx = 0; Idx < ContextualAnimAsset->DataContainer.Num(); Idx++)
 		{
-			Origin = Origin + Direction * Data.OffsetFromOrigin;
-		}
+			const FContextualAnimData& Data = ContextualAnimAsset->DataContainer[Idx];
+			const FTransform EntryTransform = Data.GetAlignmentTransformAtEntryTime() * GetComponentTransform();
 
-		// Distance Test
-		//--------------------------------------------------
-		if(Data.Distance.MaxDistance > 0.f || Data.Distance.MinDistance > 0.f)
-		{
-			const float DistSq = FVector::DistSquared2D(Origin, Querier->GetActorLocation());
+			FVector Origin = GetComponentTransform().GetLocation();
+			FVector Direction = (EntryTransform.GetLocation() - Origin).GetSafeNormal2D();
 
-			if (Data.Distance.MaxDistance > 0.f)
+			if (Data.OffsetFromOrigin != 0.f)
 			{
-				if (DistSq > FMath::Square(Data.Distance.MaxDistance))
+				Origin = Origin + Direction * Data.OffsetFromOrigin;
+			}
+
+			// Distance Test
+			//--------------------------------------------------
+			if (Data.Distance.MaxDistance > 0.f || Data.Distance.MinDistance > 0.f)
+			{
+				const float DistSq = FVector::DistSquared2D(Origin, QueryParams.Querier->GetActorLocation());
+
+				if (Data.Distance.MaxDistance > 0.f)
 				{
-					continue;
+					if (DistSq > FMath::Square(Data.Distance.MaxDistance))
+					{
+						continue;
+					}
+				}
+
+				if (Data.Distance.MinDistance > 0.f)
+				{
+					if (DistSq < FMath::Square(Data.Distance.MinDistance))
+					{
+						continue;
+					}
 				}
 			}
 
-			if (Data.Distance.MinDistance > 0.f)
+			// Angle Test
+			//--------------------------------------------------
+			if (Data.Angle.Tolerance > 0.f)
 			{
-				if (DistSq < FMath::Square(Data.Distance.MinDistance))
+				//@TODO: Cache this
+				const float AngleCos = FMath::Cos(FMath::Clamp(FMath::DegreesToRadians(Data.Angle.Tolerance), 0.f, PI));
+				const FVector ToLocation = (QueryParams.Querier->GetActorLocation() - Origin).GetSafeNormal2D();
+				if (FVector::DotProduct(ToLocation, Direction) < AngleCos)
 				{
 					continue;
 				}
 			}
-		}
 
-		// Angle Test
-		//--------------------------------------------------
-		if(Data.Angle.Tolerance > 0.f)
-		{
-			//@TODO: Cache this
-			const float AngleCos = FMath::Cos(FMath::Clamp(FMath::DegreesToRadians(Data.Angle.Tolerance), 0.f, PI));
-			const FVector ToLocation = (Querier->GetActorLocation() - Origin).GetSafeNormal2D();
-			if (FVector::DotProduct(ToLocation, Direction) < AngleCos)
+			// Facing Test
+			//--------------------------------------------------
+			if (Data.Facing.Tolerance > 0.f)
 			{
-				continue;
+				//@TODO: Cache this
+				const float FacingCos = FMath::Cos(FMath::Clamp(FMath::DegreesToRadians(Data.Facing.Tolerance), 0.f, PI));
+				if (FVector::DotProduct(QueryParams.Querier->GetActorForwardVector(), EntryTransform.GetRotation().GetForwardVector()) < FacingCos)
+				{
+					continue;
+				}
 			}
-		}
 
-		// Facing Test
-		//--------------------------------------------------
-		if(Data.Facing.Tolerance > 0.f)
+			// Return the first item that passes all tests
+			DataIndex = Idx;
+			break;
+		}
+	}
+	else // Simple Query
+	{
+		float BestDistanceSq = MAX_FLT;
+		for (int32 Idx = 0; Idx < ContextualAnimAsset->DataContainer.Num(); Idx++)
 		{
-			//@TODO: Cache this
-			const float FacingCos = FMath::Cos(FMath::Clamp(FMath::DegreesToRadians(Data.Facing.Tolerance), 0.f, PI));
-			if (FVector::DotProduct(Querier->GetActorForwardVector(), EntryTransform.GetRotation().GetForwardVector()) < FacingCos)
+			const FContextualAnimData& Data = ContextualAnimAsset->DataContainer[Idx];
+
+			//@TODO: Convert querier location to local space instead
+			const FTransform EntryTransform = Data.GetAlignmentTransformAtEntryTime() * GetComponentTransform();
+			const float DistSq = FVector::DistSquared2D(EntryTransform.GetLocation(), QueryParams.Querier->GetActorLocation());
+			if (DistSq < BestDistanceSq)
 			{
-				continue;
+				BestDistanceSq = DistSq;
+				DataIndex = Idx;
 			}
 		}
-
-		// Return the first item that passes all tests
-		return &Data;
 	}
 
-	return nullptr;
-}
-
-const FContextualAnimData* UContextualAnimComponent::FindClosestDataForActor(const AActor* Querier) const
-{
-	if (ContextualAnimAsset == nullptr || Querier == nullptr)
+	if(DataIndex != INDEX_NONE)
 	{
-		return nullptr;
-	}
+		const FContextualAnimData& ResultData = ContextualAnimAsset->DataContainer[DataIndex];
 
-	float BestDistanceSq = MAX_FLT;
-	const FContextualAnimData* BestData = nullptr;
-	for (int32 Idx = 0; Idx < ContextualAnimAsset->DataContainer.Num(); Idx++)
-	{
-		const FContextualAnimData& Data = ContextualAnimAsset->DataContainer[Idx];
-		const FTransform EntryTransform = Data.GetAlignmentTransformAtEntryTime() * GetComponentTransform();
-		const float DistSq = FVector::DistSquared2D(EntryTransform.GetLocation(), Querier->GetActorLocation());
-		if (DistSq < BestDistanceSq)
+		Result.DataIndex = DataIndex;
+		Result.Animation = ResultData.Animation;
+		Result.EntryTransform = ResultData.GetAlignmentTransformAtEntryTime() * GetComponentTransform();
+		Result.SyncTransform = ResultData.GetAlignmentTransformAtSyncTime() * GetComponentTransform();
+
+		if (QueryParams.bFindAnimStartTime)
 		{
-			BestDistanceSq = DistSq;
-			BestData = &Data;
+			const FVector LocalLocation = (QueryParams.Querier->GetActorTransform().GetRelativeTransform(GetComponentTransform())).GetLocation();
+			Result.AnimStartTime = ResultData.FindBestAnimStartTime(LocalLocation);
 		}
-	}
-
-	return BestData;
-}
-
-bool UContextualAnimComponent::FindBestEntryPointForActor(const AActor* Querier, FContextualAnimEntryPoint& Result) const
-{
-	if(const FContextualAnimData* Data = FindBestDataForActor(Querier))
-	{
-		Result.Animation = Data->Animation;
-		Result.EntryTransform = Data->GetAlignmentTransformAtEntryTime() * GetComponentTransform();
-		Result.SyncTransform = Data->GetAlignmentTransformAtSyncTime() * GetComponentTransform();
-		return true;
-	}
 
-	return false;
-}
-
-bool UContextualAnimComponent::FindClosestEntryPointForActor(const AActor* Querier, FContextualAnimEntryPoint& Result) const
-{
-	if (const FContextualAnimData* Data = FindClosestDataForActor(Querier))
-	{
-		Result.Animation = Data->Animation;
-		Result.EntryTransform = Data->GetAlignmentTransformAtEntryTime() * GetComponentTransform();
-		Result.SyncTransform = Data->GetAlignmentTransformAtSyncTime() * GetComponentTransform();
 		return true;
 	}
 
@@ -217,26 +211,58 @@ FPrimitiveSceneProxy* UContextualAnimComponent::CreateSceneProxy()
 
 			const FContextualAnimDebugParams Params = ContextualAnimComp->DebugParams;
 
-			const FTransform ToWorldTransform = FTransform(GetLocalToWorld());
+			const FMatrix& LocalToWorld = GetLocalToWorld();
+			const FTransform ToWorldTransform = FTransform(LocalToWorld);
 
 			for (int32 ViewIndex = 0; ViewIndex < Views.Num(); ViewIndex++)
 			{
 				if (VisibilityMap & (1 << ViewIndex))
 				{
+					const FSceneView* View = Views[ViewIndex];
+
+					// Taking into account the min and maximum drawing distance
+					const float DistanceSqr = (View->ViewMatrices.GetViewOrigin() - LocalToWorld.GetOrigin()).SizeSquared();
+					if (DistanceSqr < FMath::Square(GetMinDrawDistance()) || DistanceSqr > FMath::Square(GetMaxDrawDistance()))
+					{
+						continue;
+					}
+
 					FPrimitiveDrawInterface* PDI = Collector.GetPDI(ViewIndex);
 
-					FName BestAssetPathName = NAME_None;
+					// Draw collision sphere
+					const FLinearColor DrawSphereColor = GetViewSelectionColor(FColor::Red, *View, IsSelected(), IsHovered(), false, IsIndividuallySelected());
+					const float SphereRadius = ContextualAnimComp->GetScaledSphereRadius();
+
+					float AbsScaleX = LocalToWorld.GetScaledAxis(EAxis::X).Size();
+					float AbsScaleY = LocalToWorld.GetScaledAxis(EAxis::Y).Size();
+					float AbsScaleZ = LocalToWorld.GetScaledAxis(EAxis::Z).Size();
+					float MinAbsScale = FMath::Min3(AbsScaleX, AbsScaleY, AbsScaleZ);
+
+					FVector ScaledX = LocalToWorld.GetUnitAxis(EAxis::X) * MinAbsScale;
+					FVector ScaledY = LocalToWorld.GetUnitAxis(EAxis::Y) * MinAbsScale;
+					FVector ScaledZ = LocalToWorld.GetUnitAxis(EAxis::Z) * MinAbsScale;
+
+					const int32 SphereSides = FMath::Clamp<int32>(SphereRadius / 4.f, 16, 64);
+					DrawCircle(PDI, LocalToWorld.GetOrigin(), ScaledX, ScaledY, DrawSphereColor, SphereRadius, SphereSides, SDPG_World);
+					DrawCircle(PDI, LocalToWorld.GetOrigin(), ScaledX, ScaledZ, DrawSphereColor, SphereRadius, SphereSides, SDPG_World);
+					DrawCircle(PDI, LocalToWorld.GetOrigin(), ScaledY, ScaledZ, DrawSphereColor, SphereRadius, SphereSides, SDPG_World);
+
+					FContextualAnimQueryResult Result;
 					if (Params.TestActor.IsValid())
 					{
-						if (const FContextualAnimData* Data = ContextualAnimComp->FindBestDataForActor(Params.TestActor.Get()))
-						{
-							BestAssetPathName = Data->Animation.GetUniqueID().GetAssetPathName();
+						if (ContextualAnimComp->QueryData(FContextualAnimQueryParams(Params.TestActor.Get(), true, true), Result))
+						{							
+							const float Time = Result.AnimStartTime;
+							const FTransform TransformAtTime = Asset->DataContainer[Result.DataIndex].GetAlignmentTransformAtTime(Time) * ToWorldTransform;
+							DrawCoordinateSystem(PDI, TransformAtTime.GetLocation(), TransformAtTime.Rotator(), 20.f, SDPG_World, 2.f);
 						}
 					}
 
-					for(const FContextualAnimData& Data : Asset->DataContainer)
+					for (int32 Idx = 0; Idx < Asset->DataContainer.Num(); Idx++)
 					{
-						FLinearColor DrawColor = BestAssetPathName == Data.Animation.GetUniqueID().GetAssetPathName() ? FLinearColor::Red : FLinearColor::White;
+						const FContextualAnimData& Data = Asset->DataContainer[Idx];
+
+						FLinearColor DrawColor = Result.DataIndex == Idx ? FLinearColor::Red : FLinearColor::White;
 
 						// Draw Entry Point
 						const FTransform EntryTransform = (Data.GetAlignmentTransformAtEntryTime() * ToWorldTransform);
diff --git a/Engine/Plugins/Experimental/Animation/ContextualAnimation/Source/ContextualAnimation/Public/ContextualAnimAsset.h b/Engine/Plugins/Experimental/Animation/ContextualAnimation/Source/ContextualAnimation/Public/ContextualAnimAsset.h
--- a/Engine/Plugins/Experimental/Animation/ContextualAnimation/Source/ContextualAnimation/Public/ContextualAnimAsset.h
+++ b/Engine/Plugins/Experimental/Animation/ContextualAnimation/Source/ContextualAnimation/Public/ContextualAnimAsset.h
@@ -29,17 +29,17 @@ struct FAlignmentTrackContainer
 	FRawAnimSequenceTrack Track;
 
 	UPROPERTY()
-	float SampleRate;
+	float SampleInterval;
 
 	FTransform ExtractTransformAtTime(float Time) const;
 
-	void Initialize(float InSampleRate)
+	void Initialize(float InSampleInterval)
 	{
 		Track.PosKeys.Empty();
 		Track.RotKeys.Empty();
 		Track.ScaleKeys.Empty();
 
-		SampleRate = InSampleRate;
+		SampleInterval = InSampleInterval;
 	}
 };
 
@@ -84,6 +84,10 @@ struct FContextualAnimData
 	UPROPERTY(EditAnywhere, BlueprintReadWrite, Category = "Defaults", meta = (UIMin = 0, ClampMin = 0))
 	float SyncTime = 0.f;
 
+	//@TODO: Temp solution to prevent FindBestAnimStartTime from returning a time that leaves us without enough room for warping
+	UPROPERTY(EditAnywhere, BlueprintReadWrite, Category = "Defaults", meta = (UIMin = 0, ClampMin = 0))
+	float AnimMinStartTime = 0.f;
+
 	UPROPERTY(EditAnywhere, BlueprintReadWrite, Category = "Defaults")
 	float OffsetFromOrigin = 0.f;
 
@@ -102,6 +106,8 @@ struct FContextualAnimData
 	FORCEINLINE FTransform GetAlignmentTransformAtTime(float Time) const { return AlignmentData.ExtractTransformAtTime(Time); }
 	FORCEINLINE FTransform GetAlignmentTransformAtEntryTime() const { return AlignmentData.ExtractTransformAtTime(EntryTime); }
 	FORCEINLINE FTransform GetAlignmentTransformAtSyncTime() const  { return AlignmentData.ExtractTransformAtTime(SyncTime); }
+
+	float FindBestAnimStartTime(const FVector& LocalLocation) const;
 };
 
 UCLASS(Blueprintable)
@@ -114,6 +120,9 @@ public:
 	UPROPERTY(EditAnywhere, BlueprintReadWrite, Category = "Alignment")
 	FName AlignmentJoint;
 
+	UPROPERTY(EditAnywhere, BlueprintReadWrite, Category = "Alignment", meta = (ClampMin = "1", ClampMax = "60"))
+	int32 SampleRate;
+
 	UPROPERTY(EditAnywhere, BlueprintReadWrite, Category = "Alignment")
 	FTransform MeshToComponent;
 
@@ -123,9 +132,4 @@ public:
 	UContextualAnimAsset(const FObjectInitializer& ObjectInitializer);
 
 	virtual void PreSave(const class ITargetPlatform* TargetPlatform) override;
-
-#if WITH_EDITOR
-	virtual void PostEditChangeProperty(FPropertyChangedEvent& PropertyChangedEvent) override;
-#endif // WITH_EDITOR
-
 };
diff --git a/Engine/Plugins/Experimental/Animation/ContextualAnimation/Source/ContextualAnimation/Public/ContextualAnimComponent.h b/Engine/Plugins/Experimental/Animation/ContextualAnimation/Source/ContextualAnimation/Public/ContextualAnimComponent.h
--- a/Engine/Plugins/Experimental/Animation/ContextualAnimation/Source/ContextualAnimation/Public/ContextualAnimComponent.h
+++ b/Engine/Plugins/Experimental/Animation/ContextualAnimation/Source/ContextualAnimation/Public/ContextualAnimComponent.h
@@ -4,24 +4,62 @@
 
 #include "CoreMinimal.h"
 #include "Components/PrimitiveComponent.h"
+#include "Components/SphereComponent.h"
 #include "ContextualAnimComponent.generated.h"
 
 class UContextualAnimAsset;
 struct FContextualAnimData;
 
+/** Stores the result of a query function */
 USTRUCT(BlueprintType)
-struct FContextualAnimEntryPoint
+struct FContextualAnimQueryResult
 {
 	GENERATED_BODY()
 
-	UPROPERTY(EditAnywhere, BlueprintReadWrite, Category = "Defaults")
+	UPROPERTY(BlueprintReadWrite, Category = "Defaults")
 	TSoftObjectPtr<UAnimMontage> Animation;
 
-	UPROPERTY(EditAnywhere, BlueprintReadWrite, Category = "Defaults")
+	UPROPERTY(BlueprintReadWrite, Category = "Defaults")
 	FTransform EntryTransform;
 
-	UPROPERTY(EditAnywhere, BlueprintReadWrite, Category = "Defaults")
+	UPROPERTY(BlueprintReadWrite, Category = "Defaults")
 	FTransform SyncTransform;
+
+	UPROPERTY(BlueprintReadWrite, Category = "Defaults")
+	float AnimStartTime = 0.f;
+
+	UPROPERTY(BlueprintReadWrite, Category = "Defaults")
+	int32 DataIndex = INDEX_NONE;
+
+	void Reset()
+	{
+		Animation = nullptr;
+		EntryTransform = SyncTransform = FTransform::Identity;
+		AnimStartTime = 0.f;
+		DataIndex = INDEX_NONE;
+	}
+};
+
+/** Stores the parameters passed into query function */
+USTRUCT(BlueprintType)
+struct FContextualAnimQueryParams
+{
+	GENERATED_BODY()
+
+	UPROPERTY(EditAnywhere, BlueprintReadWrite, Category = "Defaults")
+	TWeakObjectPtr<const AActor> Querier;
+
+	UPROPERTY(EditAnywhere, BlueprintReadWrite, Category = "Defaults")
+	bool bComplexQuery;
+
+	UPROPERTY(EditAnywhere, BlueprintReadWrite, Category = "Defaults")
+	bool bFindAnimStartTime;
+
+	FContextualAnimQueryParams()
+		: Querier(nullptr), bComplexQuery(false), bFindAnimStartTime(false){}
+
+	FContextualAnimQueryParams(const AActor* InQuerier, bool bInComplexQuery, bool bInFindAnimStartTime)
+		: Querier(InQuerier), bComplexQuery(bInComplexQuery), bFindAnimStartTime(bInFindAnimStartTime) {}
 };
 
 USTRUCT(BlueprintType)
@@ -37,7 +75,7 @@ struct FContextualAnimDebugParams
 };
 
 UCLASS(meta = (BlueprintSpawnableComponent))
-class CONTEXTUALANIMATION_API UContextualAnimComponent : public UPrimitiveComponent
+class CONTEXTUALANIMATION_API UContextualAnimComponent : public USphereComponent
 {
 	GENERATED_BODY()
 
@@ -54,18 +92,9 @@ public:
 	UPROPERTY(EditAnywhere, BlueprintReadWrite, Category = "Debug", meta = (ShowOnlyInnerProperties, EditCondition = "bEnableDebug"))
 	FContextualAnimDebugParams DebugParams;
 
-	virtual void TickComponent(float DeltaTime, enum ELevelTick TickType, FActorComponentTickFunction* ThisTickFunction) override;
-
 	virtual FPrimitiveSceneProxy* CreateSceneProxy() override;
 	virtual FBoxSphereBounds CalcBounds(const FTransform& LocalToWorld) const;
 
-	const FContextualAnimData* FindBestDataForActor(const AActor* Querier) const;
-
-	const FContextualAnimData* FindClosestDataForActor(const AActor* Querier) const;
-
-	UFUNCTION(BlueprintCallable, Category = "Contextual Animation")
-	bool FindBestEntryPointForActor(const AActor* Querier, FContextualAnimEntryPoint& Result) const;
-
 	UFUNCTION(BlueprintCallable, Category = "Contextual Animation")
-	bool FindClosestEntryPointForActor(const AActor* Querier, FContextualAnimEntryPoint& Result) const;
+	bool QueryData(const FContextualAnimQueryParams& QueryParams, FContextualAnimQueryResult& Result) const;
 };
diff --git a/Engine/Plugins/Experimental/Animation/ContextualAnimation/Source/ContextualAnimationEditor/Private/ContextualAnimEdMode.cpp b/Engine/Plugins/Experimental/Animation/ContextualAnimation/Source/ContextualAnimationEditor/Private/ContextualAnimEdMode.cpp
--- a/Engine/Plugins/Experimental/Animation/ContextualAnimation/Source/ContextualAnimationEditor/Private/ContextualAnimEdMode.cpp
+++ b/Engine/Plugins/Experimental/Animation/ContextualAnimation/Source/ContextualAnimationEditor/Private/ContextualAnimEdMode.cpp
@@ -13,7 +13,7 @@
 #include "AIController.h"
 #include "ContextualAnimComponent.h"
 #include "MotionWarpingComponent.h"
-#include "Animation/AnimMontage.h"
+#include "Animation/AnimInstance.h"
 
 const FEditorModeID FContextualAnimEdMode::EM_ContextualAnimEdModeId = TEXT("EM_ContextualAnimEdMode");
 
@@ -78,48 +78,50 @@ bool FContextualAnimEdMode::InputKey(FEditorViewportClient* ViewportClient, FVie
 	{
 		if (!TestCharacter->IsPlayingRootMotion())
 		{
-			//@TODO: Cache the list of interactable actors
-			float BestDistance = MAX_FLT;
 			UContextualAnimComponent* ContextualAnimComp = nullptr;
-			for (TActorIterator<AActor> It(GetWorld()); It; ++It)
+
+			TArray<UPrimitiveComponent*> OverlappingComps;
+			TestCharacter->GetOverlappingComponents(OverlappingComps);
+
+			for (UPrimitiveComponent* Comp : OverlappingComps)
 			{
-				AActor* Actor = *It;
-				UContextualAnimComponent* Comp = Actor ? Actor->FindComponentByClass<UContextualAnimComponent>() : nullptr;
-				if(Comp)
+				if (Comp && Comp->GetClass()->IsChildOf<UContextualAnimComponent>())
 				{
-					const float Distance = FVector::DistSquared(Actor->GetActorLocation(), TestCharacter->GetActorLocation());
-					if(Distance < BestDistance)
-					{
-						BestDistance = Distance;
-						ContextualAnimComp = Comp;
-					}
+					ContextualAnimComp = Cast<UContextualAnimComponent>(Comp);
+					break;
 				}
 			}
 
 			if (ContextualAnimComp)
 			{
-				FContextualAnimEntryPoint EntryPoint;
-				if (ContextualAnimComp->FindBestEntryPointForActor(TestCharacter.Get(), EntryPoint))
+				FContextualAnimQueryResult Result;
+				if (ContextualAnimComp->QueryData(FContextualAnimQueryParams(TestCharacter.Get(), true, true), Result))
 				{
-					if (UAnimMontage* Montage = EntryPoint.Animation.LoadSynchronous())
+					if (UAnimMontage* Montage = Result.Animation.LoadSynchronous())
 					{
 						if (UMotionWarpingComponent* MotionWarpingComp = TestCharacter->FindComponentByClass<UMotionWarpingComponent>())
 						{
-							const FName SyncPointName = GetContextualAnimEdModeToolkit()->GetSettings()->MotionWarpSyncPointName;
-							MotionWarpingComp->AddOrUpdateSyncPoint(SyncPointName, EntryPoint.SyncTransform);
-							TestCharacter->PlayAnimMontage(Montage);
+							USkeletalMeshComponent* SkelMeshComp = TestCharacter->GetMesh();
+							UAnimInstance* AnimInstance = SkelMeshComp ? SkelMeshComp->GetAnimInstance() : nullptr;
+							if (AnimInstance)
+							{
+								const FName SyncPointName = GetContextualAnimEdModeToolkit()->GetSettings()->MotionWarpSyncPointName;
+								MotionWarpingComp->AddOrUpdateSyncPoint(SyncPointName, Result.SyncTransform);
+
+								AnimInstance->Montage_Play(Montage, 1.f, EMontagePlayReturnType::MontageLength, Result.AnimStartTime);
+							}
 						}
 					}
 				}
 				else
 				{
-					GEngine->AddOnScreenDebugMessage(-1, 5.f, FColor::Yellow, FString::Printf(TEXT("Found nearest interactable actor (%s) but the Test Character is not close enough to any of the entry points"), 
-					*GetNameSafe(ContextualAnimComp->GetOwner())));
+					GEngine->AddOnScreenDebugMessage(-1, 5.f, FColor::Yellow, FString::Printf(TEXT("Found nearest interactable actor (%s) but the Test Character is not close enough to any of the entry points"),
+						*GetNameSafe(ContextualAnimComp->GetOwner())));
 				}
 			}
 			else
 			{
-				GEngine->AddOnScreenDebugMessage(-1, 5.f, FColor::Yellow, FString(TEXT("There are no interactable actors in the level")));
+				GEngine->AddOnScreenDebugMessage(-1, 5.f, FColor::Yellow, FString(TEXT("Can't find interactable actor")));
 			}
 		}
 		else
