Contextual Anim - Moved query function to the DataAsset and added option to query the data using a just a transform. #review-15213333 [REDACTED_HANDLE] #fyi [REDACTED_HANDLE]

diff --git a/Engine/Plugins/Experimental/Animation/ContextualAnimation/Source/ContextualAnimation/Private/ContextualAnimAsset.cpp b/Engine/Plugins/Experimental/Animation/ContextualAnimation/Source/ContextualAnimation/Private/ContextualAnimAsset.cpp
--- a/Engine/Plugins/Experimental/Animation/ContextualAnimation/Source/ContextualAnimation/Private/ContextualAnimAsset.cpp
+++ b/Engine/Plugins/Experimental/Animation/ContextualAnimation/Source/ContextualAnimation/Private/ContextualAnimAsset.cpp
@@ -225,4 +225,118 @@ void UContextualAnimAsset::PreSave(const class ITargetPlatform* TargetPlatform)
 	}
 
 	Super::PreSave(TargetPlatform);
+}
+
+bool UContextualAnimAsset::QueryData(FContextualAnimQueryResult& OutResult, const FContextualAnimQueryParams& QueryParams, const FTransform& ToWorldTransform) const
+{
+	OutResult.Reset();
+
+	const FTransform QueryTransform = QueryParams.Querier.IsValid() ? QueryParams.Querier->GetActorTransform() : QueryParams.QueryTransform;
+
+	int32 DataIndex = INDEX_NONE;
+	if (QueryParams.bComplexQuery)
+	{
+		for (int32 Idx = 0; Idx < DataContainer.Num(); Idx++)
+		{
+			const FContextualAnimData& Data = DataContainer[Idx];
+			const FTransform EntryTransform = Data.GetAlignmentTransformAtEntryTime() * ToWorldTransform;
+
+			FVector Origin = ToWorldTransform.GetLocation();
+			FVector Direction = (EntryTransform.GetLocation() - Origin).GetSafeNormal2D();
+
+			if (Data.OffsetFromOrigin != 0.f)
+			{
+				Origin = Origin + Direction * Data.OffsetFromOrigin;
+			}
+
+			// Distance Test
+			//--------------------------------------------------
+			if (Data.Distance.MaxDistance > 0.f || Data.Distance.MinDistance > 0.f)
+			{
+				const float DistSq = FVector::DistSquared2D(Origin, QueryTransform.GetLocation());
+
+				if (Data.Distance.MaxDistance > 0.f)
+				{
+					if (DistSq > FMath::Square(Data.Distance.MaxDistance))
+					{
+						continue;
+					}
+				}
+
+				if (Data.Distance.MinDistance > 0.f)
+				{
+					if (DistSq < FMath::Square(Data.Distance.MinDistance))
+					{
+						continue;
+					}
+				}
+			}
+
+			// Angle Test
+			//--------------------------------------------------
+			if (Data.Angle.Tolerance > 0.f)
+			{
+				//@TODO: Cache this
+				const float AngleCos = FMath::Cos(FMath::Clamp(FMath::DegreesToRadians(Data.Angle.Tolerance), 0.f, PI));
+				const FVector ToLocation = (QueryTransform.GetLocation() - Origin).GetSafeNormal2D();
+				if (FVector::DotProduct(ToLocation, Direction) < AngleCos)
+				{
+					continue;
+				}
+			}
+
+			// Facing Test
+			//--------------------------------------------------
+			if (Data.Facing.Tolerance > 0.f)
+			{
+				//@TODO: Cache this
+				const float FacingCos = FMath::Cos(FMath::Clamp(FMath::DegreesToRadians(Data.Facing.Tolerance), 0.f, PI));
+				if (FVector::DotProduct(QueryTransform.GetRotation().GetForwardVector(), EntryTransform.GetRotation().GetForwardVector()) < FacingCos)
+				{
+					continue;
+				}
+			}
+
+			// Return the first item that passes all tests
+			DataIndex = Idx;
+			break;
+		}
+	}
+	else // Simple Query
+	{
+		float BestDistanceSq = MAX_FLT;
+		for (int32 Idx = 0; Idx < DataContainer.Num(); Idx++)
+		{
+			const FContextualAnimData& Data = DataContainer[Idx];
+
+			//@TODO: Convert querier location to local space instead
+			const FTransform EntryTransform = Data.GetAlignmentTransformAtEntryTime() * ToWorldTransform;
+			const float DistSq = FVector::DistSquared2D(EntryTransform.GetLocation(), QueryTransform.GetLocation());
+			if (DistSq < BestDistanceSq)
+			{
+				BestDistanceSq = DistSq;
+				DataIndex = Idx;
+			}
+		}
+	}
+
+	if (DataIndex != INDEX_NONE)
+	{
+		const FContextualAnimData& ResultData = DataContainer[DataIndex];
+
+		OutResult.DataIndex = DataIndex;
+		OutResult.Animation = ResultData.Animation;
+		OutResult.EntryTransform = ResultData.GetAlignmentTransformAtEntryTime() * ToWorldTransform;
+		OutResult.SyncTransform = ResultData.GetAlignmentTransformAtSyncTime() * ToWorldTransform;
+
+		if (QueryParams.bFindAnimStartTime)
+		{
+			const FVector LocalLocation = (QueryTransform.GetRelativeTransform(ToWorldTransform)).GetLocation();
+			OutResult.AnimStartTime = ResultData.FindBestAnimStartTime(LocalLocation);
+		}
+
+		return true;
+	}
+
+	return false;
 }
\ No newline at end of file
diff --git a/Engine/Plugins/Experimental/Animation/ContextualAnimation/Source/ContextualAnimation/Private/ContextualAnimComponent.cpp b/Engine/Plugins/Experimental/Animation/ContextualAnimation/Source/ContextualAnimation/Private/ContextualAnimComponent.cpp
--- a/Engine/Plugins/Experimental/Animation/ContextualAnimation/Source/ContextualAnimation/Private/ContextualAnimComponent.cpp
+++ b/Engine/Plugins/Experimental/Animation/ContextualAnimation/Source/ContextualAnimation/Private/ContextualAnimComponent.cpp
@@ -1,7 +1,6 @@
 // Copyright Epic Games, Inc. All Rights Reserved.
 
 #include "ContextualAnimComponent.h"
-#include "ContextualAnimAsset.h"
 #include "DrawDebugHelpers.h"
 
 UContextualAnimComponent::UContextualAnimComponent(const FObjectInitializer& ObjectInitializer)
@@ -20,119 +19,7 @@ UContextualAnimComponent::UContextualAnimComponent(const FObjectInitializer& Obj
 
 bool UContextualAnimComponent::QueryData(const FContextualAnimQueryParams& QueryParams, FContextualAnimQueryResult& Result) const
 {
-	Result.Reset();
-
-	if (!ContextualAnimAsset || !QueryParams.Querier.IsValid())
-	{
-		return false;
-	}
-
-	int32 DataIndex = INDEX_NONE;
-	if (QueryParams.bComplexQuery)
-	{
-		for (int32 Idx = 0; Idx < ContextualAnimAsset->DataContainer.Num(); Idx++)
-		{
-			const FContextualAnimData& Data = ContextualAnimAsset->DataContainer[Idx];
-			const FTransform EntryTransform = Data.GetAlignmentTransformAtEntryTime() * GetComponentTransform();
-
-			FVector Origin = GetComponentTransform().GetLocation();
-			FVector Direction = (EntryTransform.GetLocation() - Origin).GetSafeNormal2D();
-
-			if (Data.OffsetFromOrigin != 0.f)
-			{
-				Origin = Origin + Direction * Data.OffsetFromOrigin;
-			}
-
-			// Distance Test
-			//--------------------------------------------------
-			if (Data.Distance.MaxDistance > 0.f || Data.Distance.MinDistance > 0.f)
-			{
-				const float DistSq = FVector::DistSquared2D(Origin, QueryParams.Querier->GetActorLocation());
-
-				if (Data.Distance.MaxDistance > 0.f)
-				{
-					if (DistSq > FMath::Square(Data.Distance.MaxDistance))
-					{
-						continue;
-					}
-				}
-
-				if (Data.Distance.MinDistance > 0.f)
-				{
-					if (DistSq < FMath::Square(Data.Distance.MinDistance))
-					{
-						continue;
-					}
-				}
-			}
-
-			// Angle Test
-			//--------------------------------------------------
-			if (Data.Angle.Tolerance > 0.f)
-			{
-				//@TODO: Cache this
-				const float AngleCos = FMath::Cos(FMath::Clamp(FMath::DegreesToRadians(Data.Angle.Tolerance), 0.f, PI));
-				const FVector ToLocation = (QueryParams.Querier->GetActorLocation() - Origin).GetSafeNormal2D();
-				if (FVector::DotProduct(ToLocation, Direction) < AngleCos)
-				{
-					continue;
-				}
-			}
-
-			// Facing Test
-			//--------------------------------------------------
-			if (Data.Facing.Tolerance > 0.f)
-			{
-				//@TODO: Cache this
-				const float FacingCos = FMath::Cos(FMath::Clamp(FMath::DegreesToRadians(Data.Facing.Tolerance), 0.f, PI));
-				if (FVector::DotProduct(QueryParams.Querier->GetActorForwardVector(), EntryTransform.GetRotation().GetForwardVector()) < FacingCos)
-				{
-					continue;
-				}
-			}
-
-			// Return the first item that passes all tests
-			DataIndex = Idx;
-			break;
-		}
-	}
-	else // Simple Query
-	{
-		float BestDistanceSq = MAX_FLT;
-		for (int32 Idx = 0; Idx < ContextualAnimAsset->DataContainer.Num(); Idx++)
-		{
-			const FContextualAnimData& Data = ContextualAnimAsset->DataContainer[Idx];
-
-			//@TODO: Convert querier location to local space instead
-			const FTransform EntryTransform = Data.GetAlignmentTransformAtEntryTime() * GetComponentTransform();
-			const float DistSq = FVector::DistSquared2D(EntryTransform.GetLocation(), QueryParams.Querier->GetActorLocation());
-			if (DistSq < BestDistanceSq)
-			{
-				BestDistanceSq = DistSq;
-				DataIndex = Idx;
-			}
-		}
-	}
-
-	if(DataIndex != INDEX_NONE)
-	{
-		const FContextualAnimData& ResultData = ContextualAnimAsset->DataContainer[DataIndex];
-
-		Result.DataIndex = DataIndex;
-		Result.Animation = ResultData.Animation;
-		Result.EntryTransform = ResultData.GetAlignmentTransformAtEntryTime() * GetComponentTransform();
-		Result.SyncTransform = ResultData.GetAlignmentTransformAtSyncTime() * GetComponentTransform();
-
-		if (QueryParams.bFindAnimStartTime)
-		{
-			const FVector LocalLocation = (QueryParams.Querier->GetActorTransform().GetRelativeTransform(GetComponentTransform())).GetLocation();
-			Result.AnimStartTime = ResultData.FindBestAnimStartTime(LocalLocation);
-		}
-
-		return true;
-	}
-
-	return false;
+	return ContextualAnimAsset ? ContextualAnimAsset->QueryData(Result, QueryParams, GetComponentTransform()) : false;
 }
 
 FBoxSphereBounds UContextualAnimComponent::CalcBounds(const FTransform& LocalToWorld) const
diff --git a/Engine/Plugins/Experimental/Animation/ContextualAnimation/Source/ContextualAnimation/Public/ContextualAnimAsset.h b/Engine/Plugins/Experimental/Animation/ContextualAnimation/Source/ContextualAnimation/Public/ContextualAnimAsset.h
--- a/Engine/Plugins/Experimental/Animation/ContextualAnimation/Source/ContextualAnimation/Public/ContextualAnimAsset.h
+++ b/Engine/Plugins/Experimental/Animation/ContextualAnimation/Source/ContextualAnimation/Public/ContextualAnimAsset.h
@@ -43,6 +43,63 @@ struct FAlignmentTrackContainer
 	}
 };
 
+/** Stores the result of a query function */
+USTRUCT(BlueprintType)
+struct FContextualAnimQueryResult
+{
+	GENERATED_BODY()
+
+	UPROPERTY(BlueprintReadWrite, Category = "Defaults")
+	TSoftObjectPtr<UAnimMontage> Animation;
+
+	UPROPERTY(BlueprintReadWrite, Category = "Defaults")
+	FTransform EntryTransform;
+
+	UPROPERTY(BlueprintReadWrite, Category = "Defaults")
+	FTransform SyncTransform;
+
+	UPROPERTY(BlueprintReadWrite, Category = "Defaults")
+	float AnimStartTime = 0.f;
+
+	UPROPERTY(BlueprintReadWrite, Category = "Defaults")
+	int32 DataIndex = INDEX_NONE;
+
+	void Reset()
+	{
+		Animation = nullptr;
+		EntryTransform = SyncTransform = FTransform::Identity;
+		AnimStartTime = 0.f;
+		DataIndex = INDEX_NONE;
+	}
+};
+
+/** Stores the parameters passed into query function */
+USTRUCT(BlueprintType)
+struct FContextualAnimQueryParams
+{
+	GENERATED_BODY()
+
+	UPROPERTY(EditAnywhere, BlueprintReadWrite, Category = "Defaults")
+	TWeakObjectPtr<const AActor> Querier;
+
+	UPROPERTY(EditAnywhere, BlueprintReadWrite, Category = "Defaults")
+	FTransform QueryTransform;
+
+	UPROPERTY(EditAnywhere, BlueprintReadWrite, Category = "Defaults")
+	bool bComplexQuery = false;
+
+	UPROPERTY(EditAnywhere, BlueprintReadWrite, Category = "Defaults")
+	bool bFindAnimStartTime = false;
+
+	FContextualAnimQueryParams(){}
+
+	FContextualAnimQueryParams(const AActor* InQuerier, bool bInComplexQuery, bool bInFindAnimStartTime)
+		: Querier(InQuerier), bComplexQuery(bInComplexQuery), bFindAnimStartTime(bInFindAnimStartTime) {}
+
+	FContextualAnimQueryParams(const FTransform& InQueryTransform, bool bInComplexQuery, bool bInFindAnimStartTime)
+		: QueryTransform(InQueryTransform), bComplexQuery(bInComplexQuery), bFindAnimStartTime(bInFindAnimStartTime) {}
+};
+
 USTRUCT(BlueprintType)
 struct FContextualAnimDistanceParam
 {
@@ -132,4 +189,6 @@ public:
 	UContextualAnimAsset(const FObjectInitializer& ObjectInitializer);
 
 	virtual void PreSave(const class ITargetPlatform* TargetPlatform) override;
+
+	bool QueryData(FContextualAnimQueryResult& OutResult, const FContextualAnimQueryParams& QueryParams, const FTransform& ToWorldTransform) const;
 };
diff --git a/Engine/Plugins/Experimental/Animation/ContextualAnimation/Source/ContextualAnimation/Public/ContextualAnimComponent.h b/Engine/Plugins/Experimental/Animation/ContextualAnimation/Source/ContextualAnimation/Public/ContextualAnimComponent.h
--- a/Engine/Plugins/Experimental/Animation/ContextualAnimation/Source/ContextualAnimation/Public/ContextualAnimComponent.h
+++ b/Engine/Plugins/Experimental/Animation/ContextualAnimation/Source/ContextualAnimation/Public/ContextualAnimComponent.h
@@ -3,65 +3,10 @@
 #pragma once
 
 #include "CoreMinimal.h"
-#include "Components/PrimitiveComponent.h"
 #include "Components/SphereComponent.h"
+#include "ContextualAnimAsset.h"
 #include "ContextualAnimComponent.generated.h"
 
-class UContextualAnimAsset;
-struct FContextualAnimData;
-
-/** Stores the result of a query function */
-USTRUCT(BlueprintType)
-struct FContextualAnimQueryResult
-{
-	GENERATED_BODY()
-
-	UPROPERTY(BlueprintReadWrite, Category = "Defaults")
-	TSoftObjectPtr<UAnimMontage> Animation;
-
-	UPROPERTY(BlueprintReadWrite, Category = "Defaults")
-	FTransform EntryTransform;
-
-	UPROPERTY(BlueprintReadWrite, Category = "Defaults")
-	FTransform SyncTransform;
-
-	UPROPERTY(BlueprintReadWrite, Category = "Defaults")
-	float AnimStartTime = 0.f;
-
-	UPROPERTY(BlueprintReadWrite, Category = "Defaults")
-	int32 DataIndex = INDEX_NONE;
-
-	void Reset()
-	{
-		Animation = nullptr;
-		EntryTransform = SyncTransform = FTransform::Identity;
-		AnimStartTime = 0.f;
-		DataIndex = INDEX_NONE;
-	}
-};
-
-/** Stores the parameters passed into query function */
-USTRUCT(BlueprintType)
-struct FContextualAnimQueryParams
-{
-	GENERATED_BODY()
-
-	UPROPERTY(EditAnywhere, BlueprintReadWrite, Category = "Defaults")
-	TWeakObjectPtr<const AActor> Querier;
-
-	UPROPERTY(EditAnywhere, BlueprintReadWrite, Category = "Defaults")
-	bool bComplexQuery;
-
-	UPROPERTY(EditAnywhere, BlueprintReadWrite, Category = "Defaults")
-	bool bFindAnimStartTime;
-
-	FContextualAnimQueryParams()
-		: Querier(nullptr), bComplexQuery(false), bFindAnimStartTime(false){}
-
-	FContextualAnimQueryParams(const AActor* InQuerier, bool bInComplexQuery, bool bInFindAnimStartTime)
-		: Querier(InQuerier), bComplexQuery(bInComplexQuery), bFindAnimStartTime(bInFindAnimStartTime) {}
-};
-
 USTRUCT(BlueprintType)
 struct FContextualAnimDebugParams
 {
