Added missing includes to fix NonUnity build #jira UE-107389 #review trivial #fyi [REDACTED_HANDLE]

diff --git a/Engine/Plugins/Experimental/Animation/ContextualAnimation/Source/ContextualAnimationEditor/Private/ContextualAnimEdModeToolkit.cpp b/Engine/Plugins/Experimental/Animation/ContextualAnimation/Source/ContextualAnimationEditor/Private/ContextualAnimEdModeToolkit.cpp
--- a/Engine/Plugins/Experimental/Animation/ContextualAnimation/Source/ContextualAnimationEditor/Private/ContextualAnimEdModeToolkit.cpp
+++ b/Engine/Plugins/Experimental/Animation/ContextualAnimation/Source/ContextualAnimationEditor/Private/ContextualAnimEdModeToolkit.cpp
@@ -6,6 +6,7 @@
 #include "Widgets/Input/SButton.h"
 #include "Widgets/Text/STextBlock.h"
 #include "EditorModeManager.h"
+#include "Modules/ModuleManager.h"
 #include "ContextualAnimEdModeSettings.h"
 
 #define LOCTEXT_NAMESPACE "ContextualAnimEdModeToolkit"
diff --git a/Engine/Plugins/Experimental/Animation/ContextualAnimation/Source/ContextualAnimationEditor/Public/ContextualAnimEdModeSettings.h b/Engine/Plugins/Experimental/Animation/ContextualAnimation/Source/ContextualAnimationEditor/Public/ContextualAnimEdModeSettings.h
--- a/Engine/Plugins/Experimental/Animation/ContextualAnimation/Source/ContextualAnimationEditor/Public/ContextualAnimEdModeSettings.h
+++ b/Engine/Plugins/Experimental/Animation/ContextualAnimation/Source/ContextualAnimationEditor/Public/ContextualAnimEdModeSettings.h
@@ -3,6 +3,7 @@
 #pragma once
 
 #include "CoreMinimal.h"
+#include "Templates/SubclassOf.h"
 #include "ContextualAnimEdModeSettings.generated.h"
 
 class ACharacter;
