Contextual Anim - Re-generating the alignment track does not override the sync time any more.
Fixed FindBestAnimStartTime so it returns the time where the distance from that position to the sync point is closer to the actual distance from the actor to the sync point.
Facing Test is now based on the angle between actor forward and direction from actor to sync point
Added functions to 'simulate' start and end interactions from the EdMode

diff --git a/Engine/Plugins/Experimental/Animation/ContextualAnimation/Source/ContextualAnimation/ContextualAnimation.Build.cs b/Engine/Plugins/Experimental/Animation/ContextualAnimation/Source/ContextualAnimation/ContextualAnimation.Build.cs
--- a/Engine/Plugins/Experimental/Animation/ContextualAnimation/Source/ContextualAnimation/ContextualAnimation.Build.cs
+++ b/Engine/Plugins/Experimental/Animation/ContextualAnimation/Source/ContextualAnimation/ContextualAnimation.Build.cs
@@ -27,7 +27,8 @@ public class ContextualAnimation : ModuleRules
 			{
 				"Core",
 				"CoreUObject",
-                "Engine"
+                "Engine",
+				"MotionWarping"
 			}
 			);
 			
diff --git a/Engine/Plugins/Experimental/Animation/ContextualAnimation/Source/ContextualAnimation/Private/ContextualAnimAsset.cpp b/Engine/Plugins/Experimental/Animation/ContextualAnimation/Source/ContextualAnimation/Private/ContextualAnimAsset.cpp
--- a/Engine/Plugins/Experimental/Animation/ContextualAnimation/Source/ContextualAnimation/Private/ContextualAnimAsset.cpp
+++ b/Engine/Plugins/Experimental/Animation/ContextualAnimation/Source/ContextualAnimation/Private/ContextualAnimAsset.cpp
@@ -8,6 +8,7 @@
 #include "Animation/AnimMontage.h"
 #include "Animation/AnimationPoseData.h"
 #include "Animation/AnimTypes.h"
+#include "Containers/ArrayView.h"
 
 DEFINE_LOG_CATEGORY(LogContextualAnim);
 
@@ -77,28 +78,34 @@ float FContextualAnimData::FindBestAnimStartTime(const FVector& LocalLocation) c
 {
 	float BestTime = 0.f;
 
-	const FTransform TransformAtEntryTime = GetAlignmentTransformAtEntryTime();
-	const FVector Direction = (LocalLocation - TransformAtEntryTime.GetLocation()).GetSafeNormal2D();
-	const float Dot = FVector::DotProduct(Direction, TransformAtEntryTime.GetRotation().GetForwardVector());
-	if (Dot > 0.f)
+	if(AnimMinStartTime < 0.f)
 	{
-		float BestDistance = MAX_FLT;	
-		const TArray<FVector>& PosKeys = AlignmentData.Track.PosKeys;
+		return BestTime;
+	}
+
+	const FVector SyncPointLocation = GetAlignmentTransformAtSyncTime().GetLocation();
+	const float PerfectDistToSyncPointSq = GetAlignmentTransformAtEntryTime().GetTranslation().SizeSquared2D();
+	const float ActualDistToSyncPointSq = FVector::DistSquared2D(LocalLocation, SyncPointLocation);
+
+	if(ActualDistToSyncPointSq < PerfectDistToSyncPointSq)
+	{
+		float BestDistance = MAX_FLT;
+		TArrayView<const FVector> PosKeys(AlignmentData.Track.PosKeys.GetData(), AlignmentData.Track.PosKeys.Num());
 
-		//@TODO: Very simple search for now
+		//@TODO: Very simple search for now. Replace with Distance Matching + Pose Matching
 		for (int32 Idx = 0; Idx < PosKeys.Num(); Idx++)
 		{
 			const float Time = Idx * AlignmentData.SampleInterval;
-			if(AnimMinStartTime > 0.f && Time >= AnimMinStartTime)
+			if (AnimMinStartTime > 0.f && Time >= AnimMinStartTime)
 			{
 				break;
 			}
 
-			const float DistSquared2D = FVector::DistSquared2D(LocalLocation, PosKeys[Idx]);
-			if (DistSquared2D < BestDistance)
+			const float DistFromCurrentFrameToSyncPointSq = FVector::DistSquared2D(SyncPointLocation, PosKeys[Idx]);
+			if (DistFromCurrentFrameToSyncPointSq < ActualDistToSyncPointSq)
 			{
-				BestDistance = DistSquared2D;
 				BestTime = Time;
+				break;
 			}
 		}
 	}
@@ -114,6 +121,7 @@ UContextualAnimAsset::UContextualAnimAsset(const FObjectInitializer& ObjectIniti
 {
 	SampleRate = 15;
 	MeshToComponent = FTransform(FRotator(0.f, -90.f, 0.f));
+	MotionWarpSyncPointName = FName(TEXT("Interact"));
 }
 
 void UContextualAnimAsset::PreSave(const class ITargetPlatform* TargetPlatform)
@@ -140,10 +148,9 @@ void UContextualAnimAsset::PreSave(const class ITargetPlatform* TargetPlatform)
 		Animation->GetSectionStartAndEndTime(0, SectionStartTime, SectionEndTime);
 
 		float Time = 0.f;
-		float EndTime = SectionEndTime;
+		float EndTime = SectionEndTime; 
 		int32 SampleIndex = 0;
 
-		Data.SyncTime = EndTime;
 		Data.AlignmentData.Initialize(SampleInterval);
 		FRawAnimSequenceTrack& NewTrack = Data.AlignmentData.Track;
 
@@ -289,9 +296,12 @@ bool UContextualAnimAsset::QueryData(FContextualAnimQueryResult& OutResult, cons
 			//--------------------------------------------------
 			if (Data.Facing.Tolerance > 0.f)
 			{
+				const FTransform SyncTransform = Data.GetAlignmentTransformAtSyncTime() * ToWorldTransform;
+
 				//@TODO: Cache this
 				const float FacingCos = FMath::Cos(FMath::Clamp(FMath::DegreesToRadians(Data.Facing.Tolerance), 0.f, PI));
-				if (FVector::DotProduct(QueryTransform.GetRotation().GetForwardVector(), EntryTransform.GetRotation().GetForwardVector()) < FacingCos)
+				const FVector ToSyncPoint = (SyncTransform.GetLocation() - QueryTransform.GetLocation()).GetSafeNormal2D();
+				if (FVector::DotProduct(QueryTransform.GetRotation().GetForwardVector(), ToSyncPoint) < FacingCos)
 				{
 					continue;
 				}
diff --git a/Engine/Plugins/Experimental/Animation/ContextualAnimation/Source/ContextualAnimation/Private/ContextualAnimComponent.cpp b/Engine/Plugins/Experimental/Animation/ContextualAnimation/Source/ContextualAnimation/Private/ContextualAnimComponent.cpp
--- a/Engine/Plugins/Experimental/Animation/ContextualAnimation/Source/ContextualAnimation/Private/ContextualAnimComponent.cpp
+++ b/Engine/Plugins/Experimental/Animation/ContextualAnimation/Source/ContextualAnimation/Private/ContextualAnimComponent.cpp
@@ -2,6 +2,10 @@
 
 #include "ContextualAnimComponent.h"
 #include "DrawDebugHelpers.h"
+#include "MotionWarpingComponent.h"
+#include "Components/SkeletalMeshComponent.h"
+#include "Animation/AnimInstance.h"
+#include "GameFramework/Character.h"
 
 UContextualAnimComponent::UContextualAnimComponent(const FObjectInitializer& ObjectInitializer)
 	: Super(ObjectInitializer)
@@ -22,6 +26,185 @@ bool UContextualAnimComponent::QueryData(const FContextualAnimQueryParams& Query
 	return ContextualAnimAsset ? ContextualAnimAsset->QueryData(Result, QueryParams, GetComponentTransform()) : false;
 }
 
+UAnimInstance* UContextualAnimComponent::GetAnimInstanceForActor(AActor* Actor) const
+{
+	//@TODO: Replace with an Interface in the future
+	ACharacter* Character = Actor ? Cast<ACharacter>(Actor) : nullptr;
+	USkeletalMeshComponent* SkelMeshComp = Character ? Character->GetMesh() : nullptr;
+	return SkelMeshComp ? SkelMeshComp->GetAnimInstance() : nullptr;
+}
+
+bool UContextualAnimComponent::IsActorPlayingContextualAnimation(AActor* Actor) const
+{
+	UAnimInstance* AnimInstance = GetAnimInstanceForActor(Actor);
+	UAnimMontage* CurrentMontage = AnimInstance ? AnimInstance->GetCurrentActiveMontage() : nullptr;
+	AActor*const* ActorPtr = CurrentMontage ? MontageToActorMap.Find(CurrentMontage) : nullptr;
+	return (ActorPtr && *ActorPtr == Actor);
+}
+
+bool UContextualAnimComponent::TryStartContextualAnimation(AActor* Actor, const FContextualAnimQueryResult& Data)
+{
+	if(!Data.IsValid())
+	{
+		UE_LOG(LogContextualAnim, Warning, TEXT("TryStartContextualAnimation. QueryResult is not valid. Owner: %s Actor: %s"),
+			*GetNameSafe(GetOwner()), *GetNameSafe(Actor));
+
+		return false;
+	}
+
+	// Early out if the actor doesn't have a valid anim instance
+	UAnimInstance* AnimInstance = GetAnimInstanceForActor(Actor);
+	if (AnimInstance == nullptr)
+	{
+		UE_LOG(LogContextualAnim, Warning, TEXT("TryStartContextualAnimation. Can't find AnimInstance for the supplied actor. Owner: %s Actor: %s"),
+			*GetNameSafe(GetOwner()), *GetNameSafe(Actor));
+
+		return false;
+	}
+
+	// Early out if the actor doesn't have a motion warping component
+	UMotionWarpingComponent* MotionWarpingComp = Actor->FindComponentByClass<UMotionWarpingComponent>();
+	if (MotionWarpingComp == nullptr)
+	{
+		UE_LOG(LogContextualAnim, Warning, TEXT("TryStartContextualAnimation. Can't find MotionWarpingComp for the supplied actor. Owner: %s Actor: %s"),
+			*GetNameSafe(GetOwner()), *GetNameSafe(Actor));
+
+		return false;
+	}
+
+	// Early out if the actor is playing the animation already
+	UAnimMontage* CurrentMontage = AnimInstance->GetCurrentActiveMontage();
+	AActor** ActorPtr = CurrentMontage ? MontageToActorMap.Find(CurrentMontage) : nullptr;
+	if (ActorPtr && *ActorPtr == Actor)
+	{
+		UE_LOG(LogContextualAnim, Warning, TEXT("TryStartContextualAnimation. The supplied actor is playing the animation already. Owner: %s Actor: %s Anim: %s"), 
+		*GetNameSafe(GetOwner()), *GetNameSafe(Actor), *GetNameSafe(CurrentMontage));
+
+		return false;
+	}
+
+	// Early out if the animation is not valid
+	UAnimMontage* Montage = Data.Animation.Get();
+	if (Montage == nullptr)
+	{
+		UE_LOG(LogContextualAnim, Warning, TEXT("TryStartContextualAnimation. The animation has not been loaded yet. Owner: %s Actor: %s Anim: %s"),
+			*GetNameSafe(GetOwner()), *GetNameSafe(Actor), *Data.Animation.GetAssetName());
+
+		return false;
+	}
+
+	// Set sync point for motion warping
+	const FName SyncPointName = ContextualAnimAsset->MotionWarpSyncPointName;
+	MotionWarpingComp->AddOrUpdateSyncPoint(SyncPointName, Data.SyncTransform);
+
+	// Play animation
+	AnimInstance->Montage_Play(Montage, 1.f, EMontagePlayReturnType::MontageLength, Data.AnimStartTime);
+
+	// Listen to when the montage ends for clean up purposes
+	AnimInstance->OnMontageBlendingOut.AddUniqueDynamic(this, &UContextualAnimComponent::OnMontageBlendingOut);
+
+	// Ignore collision between actors
+	SetIgnoreOwnerComponentsWhenMovingForActor(Actor, true);
+
+	// Keep track of the actor
+	MontageToActorMap.Add(Montage, Actor);
+
+	UE_LOG(LogContextualAnim, Log, TEXT("TryStartContextualAnimation. Starting contextual anim. Owner: %s Actor: %s Anim: %s StartTime: %f SyncPointName: %s"),
+		*GetNameSafe(GetOwner()), *GetNameSafe(Actor), *GetNameSafe(Montage), Data.AnimStartTime, *SyncPointName.ToString());
+
+	return true;
+}
+
+bool UContextualAnimComponent::TryEndContextualAnimation(AActor* Actor)
+{
+	// Early out if the actor doesn't have a valid anim instance
+	UAnimInstance* AnimInstance = GetAnimInstanceForActor(Actor);
+	if(AnimInstance == nullptr)
+	{
+		UE_LOG(LogContextualAnim, Warning, TEXT("TryEndContextualAnimation. Can't find AnimInstance for the supplied actor. Owner: %s Actor: %s"),
+			*GetNameSafe(GetOwner()), *GetNameSafe(Actor));
+
+		return false;
+	}
+
+	// Early out if the actor is not playing a contextual animation
+	UAnimMontage* CurrentMontage = AnimInstance->GetCurrentActiveMontage();
+	AActor** ActorPtr = CurrentMontage ? MontageToActorMap.Find(CurrentMontage) : nullptr;
+	if(ActorPtr == nullptr || *ActorPtr != Actor)
+	{
+		UE_LOG(LogContextualAnim, Warning, TEXT("TryEndContextualAnimation. The supplied actor is not playing a contextual anim. Owner: %s Actor: %s CurrentMontage: %s"),
+			*GetNameSafe(GetOwner()), *GetNameSafe(Actor), *GetNameSafe(CurrentMontage));
+
+		return false;
+	}
+
+	// Check if we have an exit section and transition to it, otherwise just stop the montage
+	static const FName ExitSectionName = FName(TEXT("Exit"));
+	const int32 SectionIdx = CurrentMontage->GetSectionIndex(ExitSectionName);
+	if(SectionIdx != INDEX_NONE)
+	{
+		// Unbind blend out delegate for a moment so we don't get it during the transition
+		AnimInstance->OnMontageBlendingOut.RemoveDynamic(this, &UContextualAnimComponent::OnMontageBlendingOut);
+
+		AnimInstance->Montage_Play(CurrentMontage, 1.f);
+		AnimInstance->Montage_JumpToSection(ExitSectionName, CurrentMontage);
+
+		AnimInstance->OnMontageBlendingOut.AddUniqueDynamic(this, &UContextualAnimComponent::OnMontageBlendingOut);
+
+		UE_LOG(LogContextualAnim, Log, TEXT("TryEndContextualAnimation. Playing 'Exit' transition. Owner: %s Actor: %s Anim: %s"),
+			*GetNameSafe(GetOwner()), *GetNameSafe(Actor), *GetNameSafe(CurrentMontage));
+	}
+	else
+	{
+		UE_LOG(LogContextualAnim, Log, TEXT("TryEndContextualAnimation. Forcing montage to stop. Owner: %s Actor: %s Anim: %s"),
+			*GetNameSafe(GetOwner()), *GetNameSafe(Actor), *GetNameSafe(CurrentMontage));
+
+		AnimInstance->Montage_Stop(CurrentMontage->BlendOut.GetBlendTime(), CurrentMontage);
+	}
+	
+	return true;
+}
+
+void UContextualAnimComponent::OnMontageBlendingOut(UAnimMontage* Montage, bool bInterrupted)
+{
+	// Remove actor from the map
+	AActor* Actor = nullptr;
+	if(MontageToActorMap.RemoveAndCopyValue(Montage, Actor))
+	{
+		// Unbind events
+		if (UAnimInstance* AnimInstance = GetAnimInstanceForActor(Actor))
+		{
+			AnimInstance->OnMontageBlendingOut.RemoveDynamic(this, &UContextualAnimComponent::OnMontageBlendingOut);
+		}
+
+		// Restore collision between actors
+		SetIgnoreOwnerComponentsWhenMovingForActor(Actor, false);
+
+		UE_LOG(LogContextualAnim, Log, TEXT("OnMontageBlendingOut. Clean up completed. Actor: %s Montage: %s bInterrupted: %d"), *GetNameSafe(Actor), *GetNameSafe(Montage), bInterrupted);
+	}
+	else
+	{
+		UE_LOG(LogContextualAnim, Log, TEXT("OnMontageBlendingOut. Can't find actor playing this montage. Montage: %s bInterrupted: %d"), *GetNameSafe(Montage), bInterrupted);
+	}
+}
+
+void UContextualAnimComponent::SetIgnoreOwnerComponentsWhenMovingForActor(AActor* Actor, bool bShouldIgnore)
+{
+	UPrimitiveComponent* ActorRootPrimitiveComp = Actor ? Cast<UPrimitiveComponent>(Actor->GetRootComponent()) : nullptr;
+	if (ActorRootPrimitiveComp)
+	{
+		const TSet<UActorComponent*> Comps = GetOwner()->GetComponents();
+		for (UActorComponent* Comp : Comps)
+		{
+			UPrimitiveComponent* OwnerPrimitiveComp = (Comp != this) ? Cast<UPrimitiveComponent>(Comp) : nullptr;
+			if (OwnerPrimitiveComp)
+			{
+				ActorRootPrimitiveComp->IgnoreComponentWhenMoving(OwnerPrimitiveComp, bShouldIgnore);
+			}
+		}
+	}
+}
+
 FBoxSphereBounds UContextualAnimComponent::CalcBounds(const FTransform& LocalToWorld) const
 {
 	FBox BB(FVector(0.f, 0.f, -10.f), FVector(0.f, 0.f, 10.f));
diff --git a/Engine/Plugins/Experimental/Animation/ContextualAnimation/Source/ContextualAnimation/Public/ContextualAnimAsset.h b/Engine/Plugins/Experimental/Animation/ContextualAnimation/Source/ContextualAnimation/Public/ContextualAnimAsset.h
--- a/Engine/Plugins/Experimental/Animation/ContextualAnimation/Source/ContextualAnimation/Public/ContextualAnimAsset.h
+++ b/Engine/Plugins/Experimental/Animation/ContextualAnimation/Source/ContextualAnimation/Public/ContextualAnimAsset.h
@@ -71,6 +71,8 @@ struct FContextualAnimQueryResult
 		AnimStartTime = 0.f;
 		DataIndex = INDEX_NONE;
 	}
+
+	FORCEINLINE bool IsValid() const { return DataIndex != INDEX_NONE; }
 };
 
 /** Stores the parameters passed into query function */
@@ -180,6 +182,9 @@ public:
 	UPROPERTY(EditAnywhere, BlueprintReadWrite, Category = "Alignment", meta = (ClampMin = "1", ClampMax = "60"))
 	int32 SampleRate;
 
+	UPROPERTY(EditAnywhere, BlueprintReadWrite, Category = "Alignment")
+	FName MotionWarpSyncPointName;
+
 	UPROPERTY(EditAnywhere, BlueprintReadWrite, Category = "Alignment")
 	FTransform MeshToComponent;
 
diff --git a/Engine/Plugins/Experimental/Animation/ContextualAnimation/Source/ContextualAnimation/Public/ContextualAnimComponent.h b/Engine/Plugins/Experimental/Animation/ContextualAnimation/Source/ContextualAnimation/Public/ContextualAnimComponent.h
--- a/Engine/Plugins/Experimental/Animation/ContextualAnimation/Source/ContextualAnimation/Public/ContextualAnimComponent.h
+++ b/Engine/Plugins/Experimental/Animation/ContextualAnimation/Source/ContextualAnimation/Public/ContextualAnimComponent.h
@@ -7,6 +7,10 @@
 #include "ContextualAnimAsset.h"
 #include "ContextualAnimComponent.generated.h"
 
+class UAnimInstance;
+class UAnimMontage;
+class AActor;
+
 USTRUCT(BlueprintType)
 struct FContextualAnimDebugParams
 {
@@ -42,4 +46,23 @@ public:
 
 	UFUNCTION(BlueprintCallable, Category = "Contextual Animation")
 	bool QueryData(const FContextualAnimQueryParams& QueryParams, FContextualAnimQueryResult& Result) const;
+
+	bool TryStartContextualAnimation(AActor* Actor, const FContextualAnimQueryResult& Data);
+
+	bool TryEndContextualAnimation(AActor* Actor);
+
+	bool IsActorPlayingContextualAnimation(AActor* Actor) const;
+
+	void SetIgnoreOwnerComponentsWhenMovingForActor(AActor* Actor, bool bShouldIgnore);
+
+	UAnimInstance* GetAnimInstanceForActor(AActor* Actor) const;
+
+protected:
+
+	UPROPERTY()
+	TMap<UAnimMontage*, AActor*> MontageToActorMap;
+
+	UFUNCTION()
+	void OnMontageBlendingOut(UAnimMontage* Montage, bool bInterrupted);
+
 };
diff --git a/Engine/Plugins/Experimental/Animation/ContextualAnimation/Source/ContextualAnimationEditor/Private/ContextualAnimEdMode.cpp b/Engine/Plugins/Experimental/Animation/ContextualAnimation/Source/ContextualAnimationEditor/Private/ContextualAnimEdMode.cpp
--- a/Engine/Plugins/Experimental/Animation/ContextualAnimation/Source/ContextualAnimationEditor/Private/ContextualAnimEdMode.cpp
+++ b/Engine/Plugins/Experimental/Animation/ContextualAnimation/Source/ContextualAnimationEditor/Private/ContextualAnimEdMode.cpp
@@ -9,11 +9,10 @@
 #include "ContextualAnimEdModeSettings.h"
 #include "DrawDebugHelpers.h"
 #include "GameFramework/Character.h"
+#include "GameFramework/CharacterMovementComponent.h"
 #include "NavigationSystem.h"
 #include "AIController.h"
 #include "ContextualAnimComponent.h"
-#include "MotionWarpingComponent.h"
-#include "Animation/AnimInstance.h"
 
 const FEditorModeID FContextualAnimEdMode::EM_ContextualAnimEdModeId = TEXT("EM_ContextualAnimEdMode");
 
@@ -76,57 +75,52 @@ bool FContextualAnimEdMode::InputKey(FEditorViewportClient* ViewportClient, FVie
 {
 	if(Key == EKeys::Enter && Event == IE_Released && TestCharacter.IsValid() && GEditor->IsSimulatingInEditor())
 	{
-		if (!TestCharacter->IsPlayingRootMotion())
-		{
-			UContextualAnimComponent* ContextualAnimComp = nullptr;
+		UContextualAnimComponent* ContextualAnimComp = nullptr;
 
-			TArray<UPrimitiveComponent*> OverlappingComps;
-			TestCharacter->GetOverlappingComponents(OverlappingComps);
+		TArray<UPrimitiveComponent*> OverlappingComps;
+		TestCharacter->GetOverlappingComponents(OverlappingComps);
 
-			for (UPrimitiveComponent* Comp : OverlappingComps)
+		for (UPrimitiveComponent* Comp : OverlappingComps)
+		{
+			if (Comp && Comp->GetClass()->IsChildOf<UContextualAnimComponent>())
 			{
-				if (Comp && Comp->GetClass()->IsChildOf<UContextualAnimComponent>())
-				{
-					ContextualAnimComp = Cast<UContextualAnimComponent>(Comp);
-					break;
-				}
+				ContextualAnimComp = Cast<UContextualAnimComponent>(Comp);
+				break;
 			}
+		}
 
-			if (ContextualAnimComp)
-			{
-				FContextualAnimQueryResult Result;
-				if (ContextualAnimComp->QueryData(FContextualAnimQueryParams(TestCharacter.Get(), true, true), Result))
-				{
-					if (UAnimMontage* Montage = Result.Animation.LoadSynchronous())
-					{
-						if (UMotionWarpingComponent* MotionWarpingComp = TestCharacter->FindComponentByClass<UMotionWarpingComponent>())
-						{
-							USkeletalMeshComponent* SkelMeshComp = TestCharacter->GetMesh();
-							UAnimInstance* AnimInstance = SkelMeshComp ? SkelMeshComp->GetAnimInstance() : nullptr;
-							if (AnimInstance)
-							{
-								const FName SyncPointName = GetContextualAnimEdModeToolkit()->GetSettings()->MotionWarpSyncPointName;
-								MotionWarpingComp->AddOrUpdateSyncPoint(SyncPointName, Result.SyncTransform);
-
-								AnimInstance->Montage_Play(Montage, 1.f, EMontagePlayReturnType::MontageLength, Result.AnimStartTime);
-							}
-						}
-					}
-				}
-				else
-				{
-					GEngine->AddOnScreenDebugMessage(-1, 5.f, FColor::Yellow, FString::Printf(TEXT("Found nearest interactable actor (%s) but the Test Character is not close enough to any of the entry points"),
-						*GetNameSafe(ContextualAnimComp->GetOwner())));
-				}
-			}
-			else
+		if(ContextualAnimComp == nullptr)
+		{
+			GEngine->AddOnScreenDebugMessage(-1, 5.f, FColor::Yellow, FString(TEXT("WARNING: The preview actor is not overlapping any interactable")));
+			return FEdMode::InputKey(ViewportClient, Viewport, Key, Event);
+		}
+
+		// Check if the preview actor is already playing a contextual animation
+		if (!ContextualAnimComp->IsActorPlayingContextualAnimation(TestCharacter.Get()))
+		{
+			// Try to find best entry point
+			FContextualAnimQueryResult Result;
+			const bool bQueryResult = ContextualAnimComp->QueryData(FContextualAnimQueryParams(TestCharacter.Get(), true, true), Result);
+
+			// Early out of no valid entry point is found
+			if(bQueryResult == false)
 			{
-				GEngine->AddOnScreenDebugMessage(-1, 5.f, FColor::Yellow, FString(TEXT("Can't find interactable actor")));
+				GEngine->AddOnScreenDebugMessage(-1, 5.f, FColor::Yellow, FString::Printf(TEXT("WARNING: The preview actor is overlapping an interactable (%s) but there is no valid entry point for his current position"),
+					*GetNameSafe(ContextualAnimComp->GetOwner())));
+				return FEdMode::InputKey(ViewportClient, Viewport, Key, Event);
 			}
+
+			// QueryData does not load the animation and TryStartContextualAnimation expects the animation to be loaded
+			// @TODO: We may want to provide functions to query the data and load the animation asynchronous
+			Result.Animation.LoadSynchronous();
+
+			// Attempt to start the interaction
+			ContextualAnimComp->TryStartContextualAnimation(TestCharacter.Get(), Result);
 		}
 		else
 		{
-			TestCharacter->StopAnimMontage();
+			//@TODO: For now we don't care where in the animation the preview actor is
+			ContextualAnimComp->TryEndContextualAnimation(TestCharacter.Get());
 		}
 
 		return true;
@@ -137,60 +131,78 @@ bool FContextualAnimEdMode::InputKey(FEditorViewportClient* ViewportClient, FVie
 
 bool FContextualAnimEdMode::HandleClick(FEditorViewportClient* InViewportClient, HHitProxy* HitProxy, const FViewportClick& Click)
 {
-	if (Click.IsControlDown() && GEditor->IsSimulatingInEditor())
+	if (!Click.IsControlDown())
+	{
+		return FEdMode::HandleClick(InViewportClient, HitProxy, Click);
+	}
+
+	if (!GEditor->IsSimulatingInEditor())
+	{
+		GEngine->AddOnScreenDebugMessage(-1, 5.f, FColor::Yellow, FString(TEXT("WARNING. You are not in Simulating Mode")));
+		return FEdMode::HandleClick(InViewportClient, HitProxy, Click);
+	}
+
+	FSceneViewFamilyContext ViewFamily(FSceneViewFamily::ConstructionValues(InViewportClient->Viewport, InViewportClient->GetScene(), InViewportClient->EngineShowFlags).SetRealtimeUpdate(InViewportClient->IsRealtime()));
+	FSceneView* View = InViewportClient->CalcSceneView(&ViewFamily);
+	FViewportCursorLocation Cursor(View, InViewportClient, Click.GetClickPos().X, Click.GetClickPos().Y);
+	const auto ViewportType = InViewportClient->GetViewportType();
+
+	const FVector RayStart = Cursor.GetOrigin();
+	const FVector RayEnd = RayStart + Cursor.GetDirection() * HALF_WORLD_MAX;
+
+	FHitResult HitResult;
+	InViewportClient->GetWorld()->LineTraceSingleByChannel(HitResult, RayStart, RayEnd, ECC_WorldStatic, FCollisionQueryParams::DefaultQueryParam);
+
+	GEngine->AddOnScreenDebugMessage(-1, 5.f, FColor::White, FString::Printf(TEXT("INFO: HandleClick: IsValidHit: %d Actor: %s Location: %s"),
+		HitResult.IsValidBlockingHit(), *GetNameSafe(HitResult.GetActor()), *HitResult.ImpactPoint.ToString()));
+
+	if (!HitResult.IsValidBlockingHit())
 	{
-		FSceneViewFamilyContext ViewFamily(FSceneViewFamily::ConstructionValues(InViewportClient->Viewport, InViewportClient->GetScene(), InViewportClient->EngineShowFlags).SetRealtimeUpdate(InViewportClient->IsRealtime()));
-		FSceneView* View = InViewportClient->CalcSceneView(&ViewFamily);
-		FViewportCursorLocation Cursor(View, InViewportClient, Click.GetClickPos().X, Click.GetClickPos().Y);
-		const auto ViewportType = InViewportClient->GetViewportType();
-		
-		const FVector RayStart = Cursor.GetOrigin();
-		const FVector RayEnd = RayStart + Cursor.GetDirection() * HALF_WORLD_MAX;
-				
-		FHitResult HitResult;
-		InViewportClient->GetWorld()->LineTraceSingleByChannel(HitResult, RayStart, RayEnd, ECC_WorldStatic, FCollisionQueryParams::DefaultQueryParam);
-		
-		if (HitResult.IsValidBlockingHit())
+		GEngine->AddOnScreenDebugMessage(-1, 5.f, FColor::Yellow, FString(TEXT("WARNING. Can't spawn preview actor. Reason: HitResult from click event is not a valid blocking hit.")));
+		return FEdMode::HandleClick(InViewportClient, HitProxy, Click);
+	}
+
+	if (!TestCharacter.IsValid())
+	{
+		TSubclassOf<ACharacter> TestActorClass = GetContextualAnimEdModeToolkit()->GetSettings()->TestCharacterClass;
+		if (!TestActorClass)
+		{
+			GEngine->AddOnScreenDebugMessage(-1, 5.f, FColor::Yellow, FString(TEXT("WARNING. Can't spawn preview actor. Reason: TestCharacterClass is invalid")));
+			return FEdMode::HandleClick(InViewportClient, HitProxy, Click);
+		}
+
+		TestCharacter = GetWorld()->SpawnActor<ACharacter>(TestActorClass, FTransform(HitResult.ImpactPoint));
+		if (ensureAlways(TestCharacter.IsValid()))
 		{
-			if(!TestCharacter.IsValid())
+			TestCharacter->bUseControllerRotationYaw = false;
+			if (UCharacterMovementComponent* CharacterMovementComp = TestCharacter->GetCharacterMovement())
 			{
-				TSubclassOf<ACharacter> TestActorClass = GetContextualAnimEdModeToolkit()->GetSettings()->TestCharacterClass;
-				if(TestActorClass)
-				{
-					TestCharacter = GetWorld()->SpawnActor<ACharacter>(TestActorClass, FTransform(HitResult.ImpactPoint));
-					
-					if(TestCharacter.IsValid())
-					{
-						if (!TestCharacter->AIControllerClass || !TestCharacter->AIControllerClass->IsChildOf<AAIController>())
-						{
-							TestCharacter->AIControllerClass = AAIController::StaticClass();
-						}
-
-						TestCharacter->SpawnDefaultController();
-					}
-				}
-				else
-				{
-					GEngine->AddOnScreenDebugMessage(-1, 5.f, FColor::Yellow, FString(TEXT("Select a valid character class for your test actor")));
-				}
+				CharacterMovementComp->bOrientRotationToMovement = true;
+				CharacterMovementComp->bUseControllerDesiredRotation = false;
+				CharacterMovementComp->RotationRate = FRotator(0.f, 540.0, 0.f);
 			}
-			else
+
+			if (!TestCharacter->AIControllerClass || !TestCharacter->AIControllerClass->IsChildOf<AAIController>())
 			{
-				if(AAIController* Controller = Cast<AAIController>(TestCharacter->GetController()))
-				{
-					UNavigationSystemV1* NavSys = FNavigationSystem::GetCurrent<UNavigationSystemV1>(Controller->GetWorld());
-					const ANavigationData* NavData = NavSys ? NavSys->GetNavDataForProps(Controller->GetNavAgentPropertiesRef(), Controller->GetNavAgentLocation()) : nullptr;
-
-					const bool bUsePathfinding = (NavData != nullptr);
-					Controller->MoveToLocation(HitResult.ImpactPoint, 10.f, true, bUsePathfinding);
-				}
+				TestCharacter->AIControllerClass = AAIController::StaticClass();
 			}
+
+			TestCharacter->SpawnDefaultController();
 		}
+	}
+	else
+	{
+		if (AAIController* Controller = Cast<AAIController>(TestCharacter->GetController()))
+		{
+			UNavigationSystemV1* NavSys = FNavigationSystem::GetCurrent<UNavigationSystemV1>(Controller->GetWorld());
+			const ANavigationData* NavData = NavSys ? NavSys->GetNavDataForProps(Controller->GetNavAgentPropertiesRef(), Controller->GetNavAgentLocation()) : nullptr;
 
-		return true;
+			const bool bUsePathfinding = (NavData != nullptr);
+			Controller->MoveToLocation(HitResult.ImpactPoint, 10.f, true, bUsePathfinding);
+		}
 	}
 
-	return FEdMode::HandleClick(InViewportClient, HitProxy, Click);
+	return true;
 }
 
 bool FContextualAnimEdMode::UsesToolkits() const
diff --git a/Engine/Plugins/Experimental/Animation/ContextualAnimation/Source/ContextualAnimationEditor/Private/ContextualAnimEdModeSettings.cpp b/Engine/Plugins/Experimental/Animation/ContextualAnimation/Source/ContextualAnimationEditor/Private/ContextualAnimEdModeSettings.cpp
--- a/Engine/Plugins/Experimental/Animation/ContextualAnimation/Source/ContextualAnimationEditor/Private/ContextualAnimEdModeSettings.cpp
+++ b/Engine/Plugins/Experimental/Animation/ContextualAnimation/Source/ContextualAnimationEditor/Private/ContextualAnimEdModeSettings.cpp
@@ -5,5 +5,4 @@
 UContextualAnimEdModeSettings::UContextualAnimEdModeSettings(const FObjectInitializer& ObjectInitializer)
 	: Super(ObjectInitializer)
 {
-	MotionWarpSyncPointName = FName(TEXT("Interact"));
 }
\ No newline at end of file
diff --git a/Engine/Plugins/Experimental/Animation/ContextualAnimation/Source/ContextualAnimationEditor/Public/ContextualAnimEdModeSettings.h b/Engine/Plugins/Experimental/Animation/ContextualAnimation/Source/ContextualAnimationEditor/Public/ContextualAnimEdModeSettings.h
--- a/Engine/Plugins/Experimental/Animation/ContextualAnimation/Source/ContextualAnimationEditor/Public/ContextualAnimEdModeSettings.h
+++ b/Engine/Plugins/Experimental/Animation/ContextualAnimation/Source/ContextualAnimationEditor/Public/ContextualAnimEdModeSettings.h
@@ -18,8 +18,5 @@ public:
 	UPROPERTY(EditAnywhere, Category = "Settings")
 	TSubclassOf<ACharacter> TestCharacterClass;
 
-	UPROPERTY(EditAnywhere, Category = "Settings")
-	FName MotionWarpSyncPointName;
-
 	UContextualAnimEdModeSettings(const FObjectInitializer& ObjectInitializer);
 };
\ No newline at end of file
