Added missing include to fix NonUnity build #review-15391444 trivial #fyi [REDACTED_HANDLE]

diff --git a/Engine/Plugins/Experimental/Animation/ContextualAnimation/Source/ContextualAnimationEditor/Private/ContextualAnimEdMode.cpp b/Engine/Plugins/Experimental/Animation/ContextualAnimation/Source/ContextualAnimationEditor/Private/ContextualAnimEdMode.cpp
--- a/Engine/Plugins/Experimental/Animation/ContextualAnimation/Source/ContextualAnimationEditor/Private/ContextualAnimEdMode.cpp
+++ b/Engine/Plugins/Experimental/Animation/ContextualAnimation/Source/ContextualAnimationEditor/Private/ContextualAnimEdMode.cpp
@@ -4,6 +4,7 @@
 #include "ContextualAnimEdModeToolkit.h"
 #include "Toolkits/ToolkitManager.h"
 #include "EditorModeManager.h"
+#include "EditorViewportClient.h"
 #include "Engine/Selection.h"
 #include "EngineUtils.h"
 #include "ContextualAnimEdModeSettings.h"
