Contextual Anim - Added support for interactions with arbitrary number of actors #jira none #review-15491509 [REDACTED_HANDLE], [REDACTED_HANDLE] #fyi [REDACTED_HANDLE]

diff --git a/Engine/Plugins/Experimental/Animation/ContextualAnimation/Source/ContextualAnimation/Private/ContextualAnimActorInterface.cpp b/Engine/Plugins/Experimental/Animation/ContextualAnimation/Source/ContextualAnimation/Private/ContextualAnimActorInterface.cpp
new file mode 100644
--- /dev/null
+++ b/Engine/Plugins/Experimental/Animation/ContextualAnimation/Source/ContextualAnimation/Private/ContextualAnimActorInterface.cpp
@@ -0,0 +1,7 @@
+// Copyright Epic Games, Inc. All Rights Reserved.
+
+#include "ContextualAnimActorInterface.h"
+
+UContextualAnimActorInterface::UContextualAnimActorInterface(const FObjectInitializer& ObjectInitializer)
+	: Super(ObjectInitializer) 
+{}
\ No newline at end of file
diff --git a/Engine/Plugins/Experimental/Animation/ContextualAnimation/Source/ContextualAnimation/Private/ContextualAnimAsset.cpp b/Engine/Plugins/Experimental/Animation/ContextualAnimation/Source/ContextualAnimation/Private/ContextualAnimAsset.cpp
--- a/Engine/Plugins/Experimental/Animation/ContextualAnimation/Source/ContextualAnimation/Private/ContextualAnimAsset.cpp
+++ b/Engine/Plugins/Experimental/Animation/ContextualAnimation/Source/ContextualAnimation/Private/ContextualAnimAsset.cpp
@@ -1,352 +1,9 @@
 // Copyright Epic Games, Inc. All Rights Reserved.
 
 #include "ContextualAnimAsset.h"
-#include "Components/SkeletalMeshComponent.h"
-#include "DrawDebugHelpers.h"
-#include "AnimationUtils.h"
-#include "AnimationRuntime.h"
-#include "Animation/AnimMontage.h"
-#include "Animation/AnimationPoseData.h"
-#include "Animation/AnimTypes.h"
-#include "Containers/ArrayView.h"
-
-DEFINE_LOG_CATEGORY(LogContextualAnim);
-
-// UContextualAnimUtilities
-///////////////////////////////////////////////////////////////////////
-
-void UContextualAnimUtilities::ExtractLocalSpacePose(const UAnimSequenceBase* Animation, const FBoneContainer& BoneContainer, float Time, bool bExtractRootMotion, FCompactPose& OutPose)
-{
-	OutPose.SetBoneContainer(&BoneContainer);
-
-	FBlendedCurve Curve;
-	Curve.InitFrom(BoneContainer);
-
-	FAnimExtractContext Context(Time, bExtractRootMotion);
-
-	FStackCustomAttributes Attributes;
-	FAnimationPoseData AnimationPoseData(OutPose, Curve, Attributes);
-	if (const UAnimSequence* AnimSequence = Cast<UAnimSequence>(Animation))
-	{
-		AnimSequence->GetBonePose(AnimationPoseData, Context);
-	}
-	else if (const UAnimMontage* AnimMontage = Cast<UAnimMontage>(Animation))
-	{
-		const FAnimTrack& AnimTrack = AnimMontage->SlotAnimTracks[0].AnimTrack;
-		AnimTrack.GetAnimationPose(AnimationPoseData, Context);
-	}
-}
-
-void UContextualAnimUtilities::ExtractComponentSpacePose(const UAnimSequenceBase* Animation, const FBoneContainer& BoneContainer, float Time, bool bExtractRootMotion, FCSPose<FCompactPose>& OutPose)
-{
-	FCompactPose Pose;
-	ExtractLocalSpacePose(Animation, BoneContainer, Time, bExtractRootMotion, Pose);
-	OutPose.InitPose(MoveTemp(Pose));
-}
-
-FTransform UContextualAnimUtilities::ExtractRootMotionFromAnimation(const UAnimSequenceBase* Animation, float StartTime, float EndTime)
-{
-	if (const UAnimMontage* Anim = Cast<UAnimMontage>(Animation))
-	{
-		return Anim->ExtractRootMotionFromTrackRange(StartTime, EndTime);
-	}
-
-	if (const UAnimSequence* Anim = Cast<UAnimSequence>(Animation))
-	{
-		return Anim->ExtractRootMotionFromRange(StartTime, EndTime);
-	}
-
-	return FTransform::Identity;
-}
-
-// FAlignmentTrackContainer
-///////////////////////////////////////////////////////////////////////
-
-FTransform FAlignmentTrackContainer::ExtractTransformAtTime(float Time) const
-{
-	FTransform AlignmentTransform = FTransform::Identity;
-	const int32 TotalFrames = Track.PosKeys.Num();
-	const float TrackLength = (TotalFrames - 1) * SampleInterval;
-	FAnimationUtils::ExtractTransformFromTrack(Time, TotalFrames, TrackLength, Track, EAnimInterpolationType::Linear, AlignmentTransform);
-	return AlignmentTransform;
-}
-
-// FContextualAnimData
-///////////////////////////////////////////////////////////////////////
-
-float FContextualAnimData::FindBestAnimStartTime(const FVector& LocalLocation) const
-{
-	float BestTime = 0.f;
-
-	if(AnimMinStartTime < 0.f)
-	{
-		return BestTime;
-	}
-
-	const FVector SyncPointLocation = GetAlignmentTransformAtSyncTime().GetLocation();
-	const float PerfectDistToSyncPointSq = GetAlignmentTransformAtEntryTime().GetTranslation().SizeSquared2D();
-	const float ActualDistToSyncPointSq = FVector::DistSquared2D(LocalLocation, SyncPointLocation);
-
-	if(ActualDistToSyncPointSq < PerfectDistToSyncPointSq)
-	{
-		float BestDistance = MAX_FLT;
-		TArrayView<const FVector> PosKeys(AlignmentData.Track.PosKeys.GetData(), AlignmentData.Track.PosKeys.Num());
-
-		//@TODO: Very simple search for now. Replace with Distance Matching + Pose Matching
-		for (int32 Idx = 0; Idx < PosKeys.Num(); Idx++)
-		{
-			const float Time = Idx * AlignmentData.SampleInterval;
-			if (AnimMinStartTime > 0.f && Time >= AnimMinStartTime)
-			{
-				break;
-			}
-
-			const float DistFromCurrentFrameToSyncPointSq = FVector::DistSquared2D(SyncPointLocation, PosKeys[Idx]);
-			if (DistFromCurrentFrameToSyncPointSq < ActualDistToSyncPointSq)
-			{
-				BestTime = Time;
-				break;
-			}
-		}
-	}
-
-	return BestTime;
-}
-
-// UContextualAnimAsset
-///////////////////////////////////////////////////////////////////////
 
 UContextualAnimAsset::UContextualAnimAsset(const FObjectInitializer& ObjectInitializer)
 	: Super(ObjectInitializer) 
 {
-	SampleRate = 15;
-	MeshToComponent = FTransform(FRotator(0.f, -90.f, 0.f));
-	MotionWarpSyncPointName = FName(TEXT("Interact"));
-}
-
-void UContextualAnimAsset::PreSave(const class ITargetPlatform* TargetPlatform)
-{
-	const FTransform MeshToComponentInverse = MeshToComponent.Inverse();
-
-	TArray<FBoneIndexType> RequiredBoneIndexArray;
-	RequiredBoneIndexArray.Add(0);
-	FBoneContainer BoneContainer;
-
-	const float SampleInterval = 1.f / SampleRate;
-
-	for (FContextualAnimData& Data : DataContainer)
-	{
-		const UAnimMontage* Animation = Data.Animation.LoadSynchronous();
-
-		if(Animation == nullptr || !Animation->IsValidSectionIndex(0))
-		{
-			continue;
-		}
-		
-		// For now we just assume the entry animation for the interaction is the first section of the montage
-		float SectionStartTime, SectionEndTime;
-		Animation->GetSectionStartAndEndTime(0, SectionStartTime, SectionEndTime);
-
-		float Time = 0.f;
-		float EndTime = SectionEndTime; 
-		int32 SampleIndex = 0;
-
-		Data.AlignmentData.Initialize(SampleInterval);
-		FRawAnimSequenceTrack& NewTrack = Data.AlignmentData.Track;
-
-		if (AlignmentJoint.IsNone())
-		{
-			if(!BoneContainer.IsValid())
-			{
-				BoneContainer = FBoneContainer(RequiredBoneIndexArray, FCurveEvaluationOption(false), *Animation->GetSkeleton());
-			}
-
-			while (Time < EndTime)
-			{
-				Time = FMath::Clamp(SampleIndex * SampleInterval, 0.f, EndTime);
-				SampleIndex++;
-
-				FCSPose<FCompactPose> ComponentSpacePose;
-				UContextualAnimUtilities::ExtractComponentSpacePose(Animation, BoneContainer, Time, false, ComponentSpacePose);
-
-				const FTransform RootTransform = MeshToComponentInverse * ComponentSpacePose.GetComponentSpaceTransform(FCompactPoseBoneIndex(0));
-				const FTransform RootRelativeToRef = RootTransform.GetRelativeTransform(MeshToComponentInverse);
-
-				NewTrack.PosKeys.Add(RootRelativeToRef.GetLocation());
-				NewTrack.RotKeys.Add(RootRelativeToRef.GetRotation());
-				NewTrack.ScaleKeys.Add(RootRelativeToRef.GetScale3D());
-			};
-		}
-		else
-		{
-			if (!BoneContainer.IsValid())
-			{
-				const int32 AlignmentJointIdx = Animation->GetSkeleton()->GetReferenceSkeleton().FindBoneIndex(AlignmentJoint);
-				if (AlignmentJointIdx == INDEX_NONE)
-				{
-					UE_LOG(LogContextualAnim, Warning, TEXT("%s Can't generate alignment track for %s. Reason: AlignmentJoint (%s) invalid"),
-						*GetNameSafe(this), *GetNameSafe(Animation), *AlignmentJoint.ToString());
-
-					continue;
-				}
-
-				RequiredBoneIndexArray.Add(AlignmentJointIdx);
-				Animation->GetSkeleton()->GetReferenceSkeleton().EnsureParentsExistAndSort(RequiredBoneIndexArray);
-
-				BoneContainer = FBoneContainer(RequiredBoneIndexArray, FCurveEvaluationOption(false), *Animation->GetSkeleton());
-			}
-
-			while (Time < EndTime)
-			{
-				Time = FMath::Clamp(SampleIndex * SampleInterval, 0.f, EndTime);
-				SampleIndex++;
-
-				FCSPose<FCompactPose> ComponentSpacePose;
-				UContextualAnimUtilities::ExtractComponentSpacePose(Animation, BoneContainer, Time, false, ComponentSpacePose);
-
-				const FTransform RootTransform = MeshToComponentInverse * ComponentSpacePose.GetComponentSpaceTransform(FCompactPoseBoneIndex(0));
-
-				FCompactPoseBoneIndex BoneIndex = BoneContainer.MakeCompactPoseIndex(FMeshPoseBoneIndex(BoneContainer.GetReferenceSkeleton().FindBoneIndex(AlignmentJoint)));
-				const FTransform RefTransform = MeshToComponentInverse * ComponentSpacePose.GetComponentSpaceTransform(FCompactPoseBoneIndex(BoneIndex));
-
-				const FTransform RootRelativeToRef = RootTransform.GetRelativeTransform(RefTransform);
-
-				NewTrack.PosKeys.Add(RootRelativeToRef.GetLocation());
-				NewTrack.RotKeys.Add(RootRelativeToRef.GetRotation());
-				NewTrack.ScaleKeys.Add(RootRelativeToRef.GetScale3D());
-			};
-		}
-
-		const FTransform EntryTransform = Data.GetAlignmentTransformAtEntryTime();
-
-		// Reference value for Distance Test
-		Data.Distance.Value = EntryTransform.GetTranslation().Size2D();
-
-		// Reference value for Angle Test
-		const FVector Direction = EntryTransform.GetTranslation().GetSafeNormal2D();
-		FQuat Delta = FQuat::FindBetweenNormals(Direction, FVector::ForwardVector);
-		Data.Angle.Value = Delta.Rotator().Yaw;
-
-		// Reference value for Facing Test
-		Data.Facing.Value = EntryTransform.Rotator().Yaw;
-	}
-
-	Super::PreSave(TargetPlatform);
-}
-
-bool UContextualAnimAsset::QueryData(FContextualAnimQueryResult& OutResult, const FContextualAnimQueryParams& QueryParams, const FTransform& ToWorldTransform) const
-{
-	OutResult.Reset();
-
-	const FTransform QueryTransform = QueryParams.Querier.IsValid() ? QueryParams.Querier->GetActorTransform() : QueryParams.QueryTransform;
-
-	int32 DataIndex = INDEX_NONE;
-	if (QueryParams.bComplexQuery)
-	{
-		for (int32 Idx = 0; Idx < DataContainer.Num(); Idx++)
-		{
-			const FContextualAnimData& Data = DataContainer[Idx];
-			const FTransform EntryTransform = Data.GetAlignmentTransformAtEntryTime() * ToWorldTransform;
-
-			FVector Origin = ToWorldTransform.GetLocation();
-			FVector Direction = (EntryTransform.GetLocation() - Origin).GetSafeNormal2D();
-
-			if (Data.OffsetFromOrigin != 0.f)
-			{
-				Origin = Origin + Direction * Data.OffsetFromOrigin;
-			}
-
-			// Distance Test
-			//--------------------------------------------------
-			if (Data.Distance.MaxDistance > 0.f || Data.Distance.MinDistance > 0.f)
-			{
-				const float DistSq = FVector::DistSquared2D(Origin, QueryTransform.GetLocation());
-
-				if (Data.Distance.MaxDistance > 0.f)
-				{
-					if (DistSq > FMath::Square(Data.Distance.MaxDistance))
-					{
-						continue;
-					}
-				}
-
-				if (Data.Distance.MinDistance > 0.f)
-				{
-					if (DistSq < FMath::Square(Data.Distance.MinDistance))
-					{
-						continue;
-					}
-				}
-			}
-
-			// Angle Test
-			//--------------------------------------------------
-			if (Data.Angle.Tolerance > 0.f)
-			{
-				//@TODO: Cache this
-				const float AngleCos = FMath::Cos(FMath::Clamp(FMath::DegreesToRadians(Data.Angle.Tolerance), 0.f, PI));
-				const FVector ToLocation = (QueryTransform.GetLocation() - Origin).GetSafeNormal2D();
-				if (FVector::DotProduct(ToLocation, Direction) < AngleCos)
-				{
-					continue;
-				}
-			}
-
-			// Facing Test
-			//--------------------------------------------------
-			if (Data.Facing.Tolerance > 0.f)
-			{
-				const FTransform SyncTransform = Data.GetAlignmentTransformAtSyncTime() * ToWorldTransform;
-
-				//@TODO: Cache this
-				const float FacingCos = FMath::Cos(FMath::Clamp(FMath::DegreesToRadians(Data.Facing.Tolerance), 0.f, PI));
-				const FVector ToSyncPoint = (SyncTransform.GetLocation() - QueryTransform.GetLocation()).GetSafeNormal2D();
-				if (FVector::DotProduct(QueryTransform.GetRotation().GetForwardVector(), ToSyncPoint) < FacingCos)
-				{
-					continue;
-				}
-			}
-
-			// Return the first item that passes all tests
-			DataIndex = Idx;
-			break;
-		}
-	}
-	else // Simple Query
-	{
-		float BestDistanceSq = MAX_FLT;
-		for (int32 Idx = 0; Idx < DataContainer.Num(); Idx++)
-		{
-			const FContextualAnimData& Data = DataContainer[Idx];
-
-			//@TODO: Convert querier location to local space instead
-			const FTransform EntryTransform = Data.GetAlignmentTransformAtEntryTime() * ToWorldTransform;
-			const float DistSq = FVector::DistSquared2D(EntryTransform.GetLocation(), QueryTransform.GetLocation());
-			if (DistSq < BestDistanceSq)
-			{
-				BestDistanceSq = DistSq;
-				DataIndex = Idx;
-			}
-		}
-	}
-
-	if (DataIndex != INDEX_NONE)
-	{
-		const FContextualAnimData& ResultData = DataContainer[DataIndex];
-
-		OutResult.DataIndex = DataIndex;
-		OutResult.Animation = ResultData.Animation;
-		OutResult.EntryTransform = ResultData.GetAlignmentTransformAtEntryTime() * ToWorldTransform;
-		OutResult.SyncTransform = ResultData.GetAlignmentTransformAtSyncTime() * ToWorldTransform;
-
-		if (QueryParams.bFindAnimStartTime)
-		{
-			const FVector LocalLocation = (QueryTransform.GetRelativeTransform(ToWorldTransform)).GetLocation();
-			OutResult.AnimStartTime = ResultData.FindBestAnimStartTime(LocalLocation);
-		}
-
-		return true;
-	}
 
-	return false;
 }
\ No newline at end of file
diff --git a/Engine/Plugins/Experimental/Animation/ContextualAnimation/Source/ContextualAnimation/Private/ContextualAnimComponent.cpp b/Engine/Plugins/Experimental/Animation/ContextualAnimation/Source/ContextualAnimation/Private/ContextualAnimComponent.cpp
deleted file mode 100644
diff --git a/Engine/Plugins/Experimental/Animation/ContextualAnimation/Source/ContextualAnimation/Private/ContextualAnimCompositeSceneAsset.cpp b/Engine/Plugins/Experimental/Animation/ContextualAnimation/Source/ContextualAnimation/Private/ContextualAnimCompositeSceneAsset.cpp
new file mode 100644
--- /dev/null
+++ b/Engine/Plugins/Experimental/Animation/ContextualAnimation/Source/ContextualAnimation/Private/ContextualAnimCompositeSceneAsset.cpp
@@ -0,0 +1,153 @@
+// Copyright Epic Games, Inc. All Rights Reserved.
+
+#include "ContextualAnimCompositeSceneAsset.h"
+#include "ContextualAnimMetadata.h"
+#include "Animation/AnimMontage.h"
+
+const FName UContextualAnimCompositeSceneAsset::InteractorRoleName = FName(TEXT("interactor"));
+const FName UContextualAnimCompositeSceneAsset::InteractableRoleName = FName(TEXT("interactable"));
+
+UContextualAnimCompositeSceneAsset::UContextualAnimCompositeSceneAsset(const FObjectInitializer& ObjectInitializer)
+	: Super(ObjectInitializer)
+{
+	PrimaryRole = UContextualAnimCompositeSceneAsset::InteractableRoleName;
+}
+
+UClass* UContextualAnimCompositeSceneAsset::GetPreviewActorClassForRole(const FName& Role) const
+{
+	return Role == PrimaryRole ? InteractableTrack.Settings.PreviewActorClass : InteractorTrack.Settings.PreviewActorClass;
+}
+
+EContextualAnimJoinRule UContextualAnimCompositeSceneAsset::GetJoinRuleForRole(const FName& Role) const
+{
+	return Role == PrimaryRole ? InteractableTrack.Settings.JoinRule : InteractorTrack.Settings.JoinRule;
+}
+
+void UContextualAnimCompositeSceneAsset::PreSave(const class ITargetPlatform* TargetPlatform)
+{
+	Super::PreSave(TargetPlatform);
+
+	for (FContextualAnimData& Data : InteractorTrack.AnimDataContainer)
+	{
+		GenerateAlignmentTracksRelativeToScenePivot(Data);
+	}
+}
+
+bool UContextualAnimCompositeSceneAsset::QueryData(FContextualAnimQueryResult& OutResult, const FContextualAnimQueryParams& QueryParams, const FTransform& ToWorldTransform) const
+{
+	OutResult.Reset();
+
+	const FTransform QueryTransform = QueryParams.Querier.IsValid() ? QueryParams.Querier->GetActorTransform() : QueryParams.QueryTransform;
+
+	int32 DataIndex = INDEX_NONE;
+	if (QueryParams.bComplexQuery)
+	{
+		for (int32 Idx = 0; Idx < InteractorTrack.AnimDataContainer.Num(); Idx++)
+		{
+			const FContextualAnimData& Data = InteractorTrack.AnimDataContainer[Idx];
+
+			if (Data.Metadata)
+			{
+				const FTransform EntryTransform = Data.GetAlignmentTransformAtEntryTime() * ToWorldTransform;
+
+				FVector Origin = ToWorldTransform.GetLocation();
+				FVector DirToEntry = (EntryTransform.GetLocation() - Origin).GetSafeNormal2D();
+
+				if (Data.Metadata->OffsetFromOrigin != 0.f)
+				{
+					Origin = Origin + DirToEntry * Data.Metadata->OffsetFromOrigin;
+				}
+
+				// Distance Test
+				//--------------------------------------------------
+				if (Data.Metadata->Distance.MaxDistance > 0.f || Data.Metadata->Distance.MinDistance > 0.f)
+				{
+					const float DistSq = FVector::DistSquared2D(Origin, QueryTransform.GetLocation());
+
+					if (Data.Metadata->Distance.MaxDistance > 0.f)
+					{
+						if (DistSq > FMath::Square(Data.Metadata->Distance.MaxDistance))
+						{
+							continue;
+						}
+					}
+
+					if (Data.Metadata->Distance.MinDistance > 0.f)
+					{
+						if (DistSq < FMath::Square(Data.Metadata->Distance.MinDistance))
+						{
+							continue;
+						}
+					}
+				}
+
+				// Angle Test
+				//--------------------------------------------------
+				if (Data.Metadata->Angle.Tolerance > 0.f)
+				{
+					//@TODO: Cache this
+					const float AngleCos = FMath::Cos(FMath::Clamp(FMath::DegreesToRadians(Data.Metadata->Angle.Tolerance), 0.f, PI));
+					const FVector ToTarget = (QueryTransform.GetLocation() - Origin).GetSafeNormal2D();
+					if (FVector::DotProduct(ToTarget, DirToEntry) < AngleCos)
+					{
+						continue;
+					}
+				}
+
+				// Facing Test
+				//--------------------------------------------------
+				if (Data.Metadata->Facing.Tolerance > 0.f)
+				{
+					//@TODO: Cache this
+					const float FacingCos = FMath::Cos(FMath::Clamp(FMath::DegreesToRadians(Data.Metadata->Facing.Tolerance), 0.f, PI));
+					const FVector ToTarget = (ToWorldTransform.GetLocation() - QueryTransform.GetLocation()).GetSafeNormal2D();
+					if (FVector::DotProduct(QueryTransform.GetRotation().GetForwardVector(), ToTarget) < FacingCos)
+					{
+						continue;
+					}
+				}
+			}
+
+			// Return the first item that passes all tests
+			DataIndex = Idx;
+			break;
+		}
+	}
+	else // Simple Query
+	{
+		float BestDistanceSq = MAX_FLT;
+		for (int32 Idx = 0; Idx < InteractorTrack.AnimDataContainer.Num(); Idx++)
+		{
+			const FContextualAnimData& Data = InteractorTrack.AnimDataContainer[Idx];
+
+			//@TODO: Convert querier location to local space instead
+			const FTransform EntryTransform = Data.GetAlignmentTransformAtEntryTime() * ToWorldTransform;
+			const float DistSq = FVector::DistSquared2D(EntryTransform.GetLocation(), QueryTransform.GetLocation());
+			if (DistSq < BestDistanceSq)
+			{
+				BestDistanceSq = DistSq;
+				DataIndex = Idx;
+			}
+		}
+	}
+
+	if (DataIndex != INDEX_NONE)
+	{
+		const FContextualAnimData& ResultData = InteractorTrack.AnimDataContainer[DataIndex];
+
+		OutResult.DataIndex = DataIndex;
+		OutResult.Animation = ResultData.Animation;
+		OutResult.EntryTransform = ResultData.GetAlignmentTransformAtEntryTime() * ToWorldTransform;
+		OutResult.SyncTransform = ResultData.GetAlignmentTransformAtSyncTime() * ToWorldTransform;
+
+		if (QueryParams.bFindAnimStartTime)
+		{
+			const FVector LocalLocation = (QueryTransform.GetRelativeTransform(ToWorldTransform)).GetLocation();
+			OutResult.AnimStartTime = ResultData.FindBestAnimStartTime(LocalLocation);
+		}
+
+		return true;
+	}
+
+	return false;
+}
\ No newline at end of file
diff --git a/Engine/Plugins/Experimental/Animation/ContextualAnimation/Source/ContextualAnimation/Private/ContextualAnimManager.cpp b/Engine/Plugins/Experimental/Animation/ContextualAnimation/Source/ContextualAnimation/Private/ContextualAnimManager.cpp
new file mode 100644
--- /dev/null
+++ b/Engine/Plugins/Experimental/Animation/ContextualAnimation/Source/ContextualAnimation/Private/ContextualAnimManager.cpp
@@ -0,0 +1,202 @@
+// Copyright Epic Games, Inc. All Rights Reserved.
+
+#include "ContextualAnimManager.h"
+#include "ContextualAnimTypes.h"
+#include "ContextualAnimActorInterface.h"
+#include "ContextualAnimSceneInstance.h"
+#include "DrawDebugHelpers.h"
+#include "ContextualAnimSceneAsset.h"
+#include "ContextualAnimCompositeSceneAsset.h"
+#include "ContextualAnimation.h"
+#include "Engine/World.h"
+
+UContextualAnimManager::UContextualAnimManager(const FObjectInitializer& ObjectInitializer)
+	: Super(ObjectInitializer)
+{
+
+}
+
+UContextualAnimManager* UContextualAnimManager::Get(const UWorld* World)
+{
+	return (World) ? FContextualAnimationModule::GetManager(World) : nullptr;
+}
+
+UWorld* UContextualAnimManager::GetWorld()const
+{
+	return CastChecked<UWorld>(GetOuter());
+}
+
+void UContextualAnimManager::Tick(float DeltaTime)
+{
+	for (UContextualAnimSceneInstance* SceneInstance : Instances)
+	{
+		SceneInstance->Tick(DeltaTime);
+	}
+}
+
+bool UContextualAnimManager::IsActorInAnyScene(AActor* Actor) const
+{
+	if (Actor)
+	{
+		for (UContextualAnimSceneInstance* SceneInstance : Instances)
+		{
+			if (SceneInstance->IsActorInThisScene(Actor))
+			{
+				return true;
+			}
+		}
+	}
+
+	return false;
+}
+
+UContextualAnimSceneInstance* UContextualAnimManager::GetSceneWithActor(AActor* Actor)
+{
+	if (Actor)
+	{
+		for (UContextualAnimSceneInstance* SceneInstance : Instances)
+		{
+			if (SceneInstance->IsActorInThisScene(Actor))
+			{
+				return SceneInstance;
+			}
+		}
+	}
+
+	return nullptr;
+}
+
+bool UContextualAnimManager::TryStartScene(const UContextualAnimSceneAssetBase* SceneAsset, const TMap<FName, AActor*>& Bindings)
+{
+	if (SceneAsset == nullptr)
+	{
+		UE_LOG(LogContextualAnim, Warning, TEXT("UContextualAnimManager::TryStartScene. Can't start scene. Reason: Invalid Scene Asset"));
+		return false;
+	}
+
+	TMap<FName, FContextualAnimSceneActorData> SceneActorMap;
+	SceneActorMap.Reserve(Bindings.Num());
+
+	// @TODO: Add virtual function in the asset class that resolve the bindings for us
+	if (const UContextualAnimSceneAsset* SimpleScene = Cast<const UContextualAnimSceneAsset>(SceneAsset))
+	{
+		for (const auto& Pair : Bindings)
+		{
+			FName RoleToBind = Pair.Key;
+			AActor* ActorToBind = Pair.Value;
+
+			if(ActorToBind == nullptr)
+			{
+				UE_LOG(LogTemp, Warning, TEXT("UContextualAnimManager::TryStartScene. Can't start scene. Reason: Attempting to bind an invalid actor. SceneAsset: %s Role: %s"), 
+				*GetNameSafe(SceneAsset), *RoleToBind.ToString());
+
+				return false;
+			}
+
+			const FContextualAnimTrack* Track = SimpleScene->FindTrack(RoleToBind);
+			if(Track == nullptr)
+			{
+				UE_LOG(LogContextualAnim, Warning, TEXT("UContextualAnimManager::TryStartScene. Can't start scene. Reason: Invalid Track. SceneAsset: %s Role: %s"),
+				*GetNameSafe(SceneAsset), *RoleToBind.ToString());
+
+				return false;
+			}
+
+			SceneActorMap.Add(RoleToBind, FContextualAnimSceneActorData(ActorToBind, &Track->AnimData));
+		}
+	}
+	else if (const UContextualAnimCompositeSceneAsset* CompositeScene = Cast<const UContextualAnimCompositeSceneAsset>(SceneAsset))
+	{
+		// Bind the Primary Role first so we can perform the query relative to it
+		FContextualAnimSceneActorData PrimaryData;
+		for (const auto& Pair : Bindings)
+		{
+			if (Pair.Key == CompositeScene->PrimaryRole)
+			{
+				PrimaryData.Actor = Pair.Value;
+				PrimaryData.AnimDataPtr = &CompositeScene->InteractableTrack.AnimData;
+
+				if (!PrimaryData.Actor.IsValid())
+				{
+					UE_LOG(LogContextualAnim, Warning, TEXT("UContextualAnimManager::TryStartScene. Can't start scene. Reason: Attempting to bind an invalid actor to the primary role. SceneAsset: %s Role: %s"),
+						*GetNameSafe(SceneAsset), *CompositeScene->PrimaryRole.ToString());
+
+					return false;
+				}
+
+				if (PrimaryData.AnimDataPtr == nullptr)
+				{
+					UE_LOG(LogContextualAnim, Warning, TEXT("UContextualAnimManager::TryStartScene. Can't start scene. Reason: Invalid Track. SceneAsset: %s Role: %s"),
+						*GetNameSafe(SceneAsset), *CompositeScene->PrimaryRole.ToString());
+
+					return false;
+				}
+
+				SceneActorMap.Add(Pair.Key, PrimaryData);
+
+				break;
+			}
+		}
+
+		for (const auto& Pair : Bindings)
+		{
+			FName RoleToBind = Pair.Key;
+			AActor* ActorToBind = Pair.Value;
+
+			if (ActorToBind == nullptr)
+			{
+				UE_LOG(LogTemp, Warning, TEXT("UContextualAnimManager::TryStartScene. Can't start scene. Reason: Attempting to bind an invalid actor. SceneAsset: %s Role: %s"),
+					*GetNameSafe(SceneAsset), *RoleToBind.ToString());
+
+				return false;
+			}
+
+			// Primary role was already bound
+			if (RoleToBind != CompositeScene->PrimaryRole)
+			{
+				FTransform ToWorldTransform = PrimaryData.GetTransform();
+
+				FContextualAnimQueryResult Result;
+				const bool bQueryResult = CompositeScene->QueryData(Result, FContextualAnimQueryParams(ActorToBind, true, true), ToWorldTransform);
+				if (bQueryResult == false)
+				{
+					UE_LOG(LogContextualAnim, Warning, TEXT("UContextualAnimManager::TryStartScene. Can't start scene. Reason: Can't find Track for actor. SceneAsset: %s Role: %s Actor: %s"),
+						*GetNameSafe(SceneAsset), *RoleToBind.ToString(), *GetNameSafe(ActorToBind));
+
+					return false;
+				}
+
+				SceneActorMap.Add(RoleToBind, FContextualAnimSceneActorData(ActorToBind, &CompositeScene->InteractorTrack.AnimDataContainer[Result.DataIndex], Result.AnimStartTime));
+			}
+		}
+	}
+
+	UContextualAnimSceneInstance* NewInstance = NewObject<UContextualAnimSceneInstance>(this);
+	NewInstance->SceneAsset = SceneAsset;
+	NewInstance->SceneActorMap = MoveTemp(SceneActorMap);
+	NewInstance->Start();
+	NewInstance->OnSceneEnded.BindUObject(this, &UContextualAnimManager::OnSceneInstanceEnded);
+
+	Instances.Add(NewInstance);
+
+	return true;
+}
+
+bool UContextualAnimManager::TryStopSceneWithActor(AActor* Actor)
+{
+	if(UContextualAnimSceneInstance* SceneInstance = GetSceneWithActor(Actor))
+	{
+		SceneInstance->Stop();
+		return true;
+	}
+
+	return false;
+}
+
+void UContextualAnimManager::OnSceneInstanceEnded(UContextualAnimSceneInstance* SceneInstance)
+{
+	if(SceneInstance)
+	{
+		Instances.Remove(SceneInstance);
+	}
+}
\ No newline at end of file
diff --git a/Engine/Plugins/Experimental/Animation/ContextualAnimation/Source/ContextualAnimation/Private/ContextualAnimMetadata.cpp b/Engine/Plugins/Experimental/Animation/ContextualAnimation/Source/ContextualAnimation/Private/ContextualAnimMetadata.cpp
new file mode 100644
--- /dev/null
+++ b/Engine/Plugins/Experimental/Animation/ContextualAnimation/Source/ContextualAnimation/Private/ContextualAnimMetadata.cpp
@@ -0,0 +1,17 @@
+// Copyright Epic Games, Inc. All Rights Reserved.
+
+#include "ContextualAnimMetadata.h"
+#include "ContextualAnimSceneAsset.h"
+#include "Engine/World.h"
+#include "DrawDebugHelpers.h"
+
+UContextualAnimMetadata::UContextualAnimMetadata(const FObjectInitializer& ObjectInitializer)
+	: Super(ObjectInitializer)
+{
+
+}
+
+UContextualAnimSceneAssetBase* UContextualAnimMetadata::GetSceneAssetOwner() const
+{
+	return Cast<UContextualAnimSceneAssetBase>(GetOuter());
+}
\ No newline at end of file
diff --git a/Engine/Plugins/Experimental/Animation/ContextualAnimation/Source/ContextualAnimation/Private/ContextualAnimSceneActorComponent.cpp b/Engine/Plugins/Experimental/Animation/ContextualAnimation/Source/ContextualAnimation/Private/ContextualAnimSceneActorComponent.cpp
new file mode 100644
--- /dev/null
+++ b/Engine/Plugins/Experimental/Animation/ContextualAnimation/Source/ContextualAnimation/Private/ContextualAnimSceneActorComponent.cpp
@@ -0,0 +1,248 @@
+// Copyright Epic Games, Inc. All Rights Reserved.
+
+#include "ContextualAnimSceneActorComponent.h"
+#include "Animation/AnimInstance.h"
+#include "GameFramework/Character.h"
+#include "ContextualAnimMetadata.h"
+
+UContextualAnimSceneActorComponent::UContextualAnimSceneActorComponent(const FObjectInitializer& ObjectInitializer)
+	: Super(ObjectInitializer)
+{
+	PrimaryComponentTick.bCanEverTick = true;
+	PrimaryComponentTick.bStartWithTickEnabled = false;
+
+	SetCollisionEnabled(ECollisionEnabled::QueryOnly);
+	SetCollisionResponseToAllChannels(ECR_Ignore);
+	SetCollisionResponseToChannel(ECC_Pawn, ECR_Overlap);
+	SetGenerateOverlapEvents(true);
+
+	bHiddenInGame = false;
+}
+
+bool UContextualAnimSceneActorComponent::QueryData(const FContextualAnimQueryParams& QueryParams, FContextualAnimQueryResult& Result) const
+{
+	return SceneAsset ? SceneAsset->QueryData(Result, QueryParams, GetComponentTransform()) : false;
+}
+
+FBoxSphereBounds UContextualAnimSceneActorComponent::CalcBounds(const FTransform& LocalToWorld) const
+{
+	FBox BB(FVector(0.f, 0.f, -10.f), FVector(0.f, 0.f, 10.f));
+	if (SceneAsset)
+	{
+		for (const FContextualAnimData& Data : SceneAsset->InteractorTrack.AnimDataContainer)
+		{
+			BB += Data.GetAlignmentTransformAtEntryTime().GetLocation();
+		}
+	}
+	return FBoxSphereBounds(BB.TransformBy(GetComponentTransform()));
+}
+
+FPrimitiveSceneProxy* UContextualAnimSceneActorComponent::CreateSceneProxy()
+{
+	class FSceneActorCompProxy final : public FPrimitiveSceneProxy
+	{
+	public:
+
+		SIZE_T GetTypeHash() const override
+		{
+			static size_t UniquePointer;
+			return reinterpret_cast<size_t>(&UniquePointer);
+		}
+
+		FSceneActorCompProxy(const UContextualAnimSceneActorComponent* InComponent)
+			: FPrimitiveSceneProxy(InComponent)
+			, SceneAssetPtr(InComponent->SceneAsset)
+			, Params(InComponent->DebugParams)
+			, Radius(InComponent->GetScaledSphereRadius())
+		{
+		}
+
+		static void DrawSector(FPrimitiveDrawInterface* PDI, const FVector& Origin, const FVector& Direction, float MinDistance, float MaxDistance, float MinAngle, float MaxAngle, const FLinearColor& Color, uint8 DepthPriority, float Thickness)
+		{
+			// Draw Cone lines
+			const FVector LeftDirection = Direction.RotateAngleAxis(MinAngle, FVector::UpVector);
+			const FVector RightDirection = Direction.RotateAngleAxis(MaxAngle, FVector::UpVector);
+			PDI->DrawLine(Origin + (LeftDirection * MinDistance), Origin + (LeftDirection * MaxDistance), Color, DepthPriority, Thickness);
+			PDI->DrawLine(Origin + (RightDirection * MinDistance), Origin + (RightDirection * MaxDistance), Color, DepthPriority, Thickness);
+
+			// Draw Near Arc
+			FVector LastDirection = LeftDirection;
+			float Angle = MinAngle;
+			while (Angle < MaxAngle)
+			{
+				Angle = FMath::Clamp<float>(Angle + 10, MinAngle, MaxAngle);
+
+				const float Length = MinDistance;
+				const FVector NewDirection = Direction.RotateAngleAxis(Angle, FVector::UpVector);
+				const FVector LineStart = Origin + (LastDirection * Length);
+				const FVector LineEnd = Origin + (NewDirection * Length);
+				PDI->DrawLine(LineStart, LineEnd, Color, DepthPriority, Thickness);
+				LastDirection = NewDirection;
+			}
+
+			// Draw Far Arc
+			LastDirection = LeftDirection;
+			Angle = MinAngle;
+			while (Angle < MaxAngle)
+			{
+				Angle = FMath::Clamp<float>(Angle + 10, MinAngle, MaxAngle);
+
+				const float Length = MaxDistance;
+				const FVector NewDirection = Direction.RotateAngleAxis(Angle, FVector::UpVector);
+				const FVector LineStart = Origin + (LastDirection * Length);
+				const FVector LineEnd = Origin + (NewDirection * Length);
+				PDI->DrawLine(LineStart, LineEnd, Color, DepthPriority, Thickness);
+				LastDirection = NewDirection;
+			}
+		}
+
+		virtual void GetDynamicMeshElements(const TArray<const FSceneView*>& Views, const FSceneViewFamily& ViewFamily, uint32 VisibilityMap, FMeshElementCollector& Collector) const override
+		{
+			const UContextualAnimCompositeSceneAsset* Asset = SceneAssetPtr.Get();
+			if (Asset == nullptr)
+			{
+				return;
+			}
+
+			const FMatrix& LocalToWorld = GetLocalToWorld();
+			const FTransform ToWorldTransform = FTransform(LocalToWorld);
+
+			for (int32 ViewIndex = 0; ViewIndex < Views.Num(); ViewIndex++)
+			{
+				if (VisibilityMap & (1 << ViewIndex))
+				{
+					const FSceneView* View = Views[ViewIndex];
+
+					// Taking into account the min and maximum drawing distance
+					const float DistanceSqr = (View->ViewMatrices.GetViewOrigin() - LocalToWorld.GetOrigin()).SizeSquared();
+					if (DistanceSqr < FMath::Square(GetMinDrawDistance()) || DistanceSqr > FMath::Square(GetMaxDrawDistance()))
+					{
+						continue;
+					}
+
+					FPrimitiveDrawInterface* PDI = Collector.GetPDI(ViewIndex);
+
+					// Draw collision sphere
+					const FLinearColor DrawSphereColor = GetViewSelectionColor(FColor::Red, *View, IsSelected(), IsHovered(), false, IsIndividuallySelected());
+
+					float AbsScaleX = LocalToWorld.GetScaledAxis(EAxis::X).Size();
+					float AbsScaleY = LocalToWorld.GetScaledAxis(EAxis::Y).Size();
+					float AbsScaleZ = LocalToWorld.GetScaledAxis(EAxis::Z).Size();
+					float MinAbsScale = FMath::Min3(AbsScaleX, AbsScaleY, AbsScaleZ);
+
+					FVector ScaledX = LocalToWorld.GetUnitAxis(EAxis::X) * MinAbsScale;
+					FVector ScaledY = LocalToWorld.GetUnitAxis(EAxis::Y) * MinAbsScale;
+					FVector ScaledZ = LocalToWorld.GetUnitAxis(EAxis::Z) * MinAbsScale;
+
+					const int32 SphereSides = FMath::Clamp<int32>(Radius / 4.f, 16, 64);
+					DrawCircle(PDI, LocalToWorld.GetOrigin(), ScaledX, ScaledY, DrawSphereColor, Radius, SphereSides, SDPG_World);
+					DrawCircle(PDI, LocalToWorld.GetOrigin(), ScaledX, ScaledZ, DrawSphereColor, Radius, SphereSides, SDPG_World);
+					DrawCircle(PDI, LocalToWorld.GetOrigin(), ScaledY, ScaledZ, DrawSphereColor, Radius, SphereSides, SDPG_World);
+
+					FContextualAnimQueryResult Result;
+					if (Params.TestActor.IsValid())
+					{
+						if (Asset->QueryData(Result, FContextualAnimQueryParams(Params.TestActor.Get(), true, true), ToWorldTransform))
+						{							
+							const float Time = Result.AnimStartTime;
+							const FTransform TransformAtTime = Asset->InteractorTrack.AnimDataContainer[Result.DataIndex].GetAlignmentTransformAtTime(Time) * ToWorldTransform;
+							DrawCoordinateSystem(PDI, TransformAtTime.GetLocation(), TransformAtTime.Rotator(), 20.f, SDPG_World, 2.f);
+						}
+					}
+
+					for (int32 Idx = 0; Idx < Asset->InteractorTrack.AnimDataContainer.Num(); Idx++)
+					{
+						const FContextualAnimData& Data = Asset->InteractorTrack.AnimDataContainer[Idx];
+
+						FLinearColor DrawColor = Result.DataIndex == Idx ? FLinearColor::Red : FLinearColor::White;
+
+						// Draw Entry Point
+						const FTransform EntryTransform = (Data.GetAlignmentTransformAtEntryTime() * ToWorldTransform);
+						DrawCoordinateSystem(PDI, EntryTransform.GetLocation(), EntryTransform.Rotator(), 20.f, SDPG_World, 3.f);
+
+						// Draw Facing Tolerance
+						if(Data.Metadata && Data.Metadata->Facing.Tolerance > 0.f)
+						{
+							DrawSector(PDI, EntryTransform.GetLocation(), EntryTransform.GetRotation().GetForwardVector(), 0.f, 30.f, -Data.Metadata->Facing.Tolerance, Data.Metadata->Facing.Tolerance, DrawColor, SDPG_World, 1.f);
+						}
+						else
+						{
+							DrawCircle(PDI, EntryTransform.GetLocation(), FVector(1, 0, 0), FVector(0, 1, 0), DrawColor, 30.f, 12, SDPG_World, 1.f);
+						}
+
+						// Draw Sync Point
+						const FTransform SyncPoint = Data.GetAlignmentTransformAtSyncTime() * ToWorldTransform;
+						DrawCoordinateSystem(PDI, SyncPoint.GetLocation(), SyncPoint.Rotator(), 20.f, SDPG_World, 3.f);
+
+						if(Params.DrawAlignmentTransformAtTime != 0.f)
+						{
+							const FTransform RootAtTime = (Data.GetAlignmentTransformAtTime(Params.DrawAlignmentTransformAtTime) * ToWorldTransform);
+							DrawCoordinateSystem(PDI, RootAtTime.GetLocation(), RootAtTime.Rotator(), 10.f, SDPG_World, 2.f);
+						}
+
+						// Draw Sector
+						if(Data.Metadata)
+						{
+							FVector Origin = ToWorldTransform.GetLocation();
+							FVector Direction = (EntryTransform.GetLocation() - ToWorldTransform.GetLocation()).GetSafeNormal2D();
+
+							if (Data.Metadata->OffsetFromOrigin != 0.f)
+							{
+								Origin = Origin + Direction * Data.Metadata->OffsetFromOrigin;
+							}
+
+							if (Data.Metadata->Angle.Tolerance > 0.f)
+							{
+								DrawSector(PDI, Origin, Direction, Data.Metadata->Distance.MinDistance, Data.Metadata->Distance.MaxDistance, -Data.Metadata->Angle.Tolerance, Data.Metadata->Angle.Tolerance, DrawColor, SDPG_World, 3.f);
+							}
+							else
+							{
+								if (Data.Metadata->Distance.MinDistance > 0.f)
+								{
+									DrawCircle(PDI, Origin, FVector(1, 0, 0), FVector(0, 1, 0), DrawColor, Data.Metadata->Distance.MinDistance, 12, SDPG_World, 2.f);
+								}
+
+								if (Data.Metadata->Distance.MaxDistance > 0.f)
+								{
+									DrawCircle(PDI, Origin, FVector(1, 0, 0), FVector(0, 1, 0), DrawColor, Data.Metadata->Distance.MaxDistance, 12, SDPG_World, 2.f);
+								}
+							}
+						}
+					}
+				}
+			}
+		}
+
+		virtual FPrimitiveViewRelevance GetViewRelevance(const FSceneView* View) const override
+		{
+			const bool bShowForCollision = View->Family->EngineShowFlags.Collision;
+			FPrimitiveViewRelevance Result;
+			Result.bDrawRelevance = IsShown(View);
+			Result.bDynamicRelevance = true;
+			Result.bNormalTranslucency = Result.bSeparateTranslucency = IsShown(View);
+			return Result;
+		}
+
+		virtual uint32 GetMemoryFootprint(void) const override
+		{
+			return(sizeof(*this) + GetAllocatedSize());
+		}
+
+		uint32 GetAllocatedSize(void) const
+		{
+			return(FPrimitiveSceneProxy::GetAllocatedSize());
+		}
+
+	private:
+		TWeakObjectPtr<const UContextualAnimCompositeSceneAsset> SceneAssetPtr;
+		FContextualAnimDebugParams Params;
+		float Radius;
+	};
+
+	if(bEnableDebug)
+	{
+		return new FSceneActorCompProxy(this);
+	}
+
+	return nullptr;
+}
\ No newline at end of file
diff --git a/Engine/Plugins/Experimental/Animation/ContextualAnimation/Source/ContextualAnimation/Private/ContextualAnimSceneAsset.cpp b/Engine/Plugins/Experimental/Animation/ContextualAnimation/Source/ContextualAnimation/Private/ContextualAnimSceneAsset.cpp
new file mode 100644
--- /dev/null
+++ b/Engine/Plugins/Experimental/Animation/ContextualAnimation/Source/ContextualAnimation/Private/ContextualAnimSceneAsset.cpp
@@ -0,0 +1,107 @@
+// Copyright Epic Games, Inc. All Rights Reserved.
+
+#include "ContextualAnimSceneAsset.h"
+#include "AnimationRuntime.h"
+#include "Animation/AnimMontage.h"
+#include "Animation/AnimationPoseData.h"
+#include "Animation/AnimTypes.h"
+#include "ContextualAnimUtilities.h"
+
+UContextualAnimSceneAsset::UContextualAnimSceneAsset(const FObjectInitializer& ObjectInitializer)
+	: Super(ObjectInitializer)
+{
+}
+
+void UContextualAnimSceneAsset::PreSave(const class ITargetPlatform* TargetPlatform)
+{
+	Super::PreSave(TargetPlatform);
+	
+	const FTransform MeshToComponentInverse = MeshToComponent.Inverse();
+	const float SampleInterval = 1.f / SampleRate;
+
+	TArray<FBoneIndexType> RequiredBoneIndexArray = { 0 };
+
+	for (auto& Entry : DataContainer)
+	{
+		// Generate alignment tracks relative to scene pivot
+		GenerateAlignmentTracksRelativeToScenePivot(Entry.Value.AnimData);
+
+		// Generate alignment tracks relative to other actors
+		const FName& Role = Entry.Key;
+		FContextualAnimData& AnimData = Entry.Value.AnimData;
+
+		const UAnimMontage* Animation = AnimData.Animation;
+
+		UE_LOG(LogContextualAnim, Log, TEXT("%s Generating Alignment Track Relative To Others for '%s' (Anim: %s)"), *GetNameSafe(this), *Role.ToString(), *GetNameSafe(Animation));
+
+		if (Animation)
+		{
+			FBoneContainer BoneContainer = FBoneContainer(RequiredBoneIndexArray, FCurveEvaluationOption(false), *Animation->GetSkeleton());
+
+			for (const auto& OtherEntry : DataContainer)
+			{
+				const FName& OtherRole = OtherEntry.Key;
+				const FContextualAnimData& OtherAnimData = OtherEntry.Value.AnimData;
+
+				if (Role != OtherRole)
+				{
+					float Time = 0.f;
+					float EndTime = Animation->GetPlayLength();
+					int32 SampleIndex = 0;
+
+					AnimData.AlignmentData.Tracks.TrackNames.Add(OtherRole);
+					FRawAnimSequenceTrack& NewTrack = AnimData.AlignmentData.Tracks.AnimationTracks.AddZeroed_GetRef();
+
+					const UAnimMontage* OtherAnimation = OtherAnimData.Animation;
+					if (OtherAnimation)
+					{
+						FBoneContainer OtherBoneContainer = FBoneContainer(RequiredBoneIndexArray, FCurveEvaluationOption(false), *OtherAnimation->GetSkeleton());
+
+						while (Time < EndTime)
+						{
+							Time = FMath::Clamp(SampleIndex * SampleInterval, 0.f, EndTime);
+							SampleIndex++;
+
+							FCSPose<FCompactPose> ComponentSpacePose;
+							UContextualAnimUtilities::ExtractComponentSpacePose(Animation, BoneContainer, Time, false, ComponentSpacePose);
+
+							FCSPose<FCompactPose> OtherComponentSpacePose;
+							UContextualAnimUtilities::ExtractComponentSpacePose(OtherAnimation, OtherBoneContainer, Time, false, OtherComponentSpacePose);
+
+							const FTransform RootTransform = MeshToComponentInverse * (ComponentSpacePose.GetComponentSpaceTransform(FCompactPoseBoneIndex(0)) * AnimData.MeshToScene);
+							const FTransform OtherRootTransform = MeshToComponentInverse * (OtherComponentSpacePose.GetComponentSpaceTransform(FCompactPoseBoneIndex(0)) * OtherAnimData.MeshToScene);
+							const FTransform RootRelativeToRef = RootTransform.GetRelativeTransform(OtherRootTransform);
+
+							NewTrack.PosKeys.Add(RootRelativeToRef.GetLocation());
+							NewTrack.RotKeys.Add(RootRelativeToRef.GetRotation());
+							NewTrack.ScaleKeys.Add(RootRelativeToRef.GetScale3D());
+						};
+					}
+
+					UE_LOG(LogContextualAnim, Log, TEXT("\t Relative To: '%s' (Anim: %s) NumKeys: %d"), *OtherRole.ToString(), *GetNameSafe(OtherAnimation), NewTrack.PosKeys.Num());
+				}
+			}
+		}
+	}
+}
+
+const FContextualAnimTrack* UContextualAnimSceneAsset::FindTrack(const FName& Role) const
+{
+	return DataContainer.Find(Role);
+}
+
+UClass* UContextualAnimSceneAsset::GetPreviewActorClassForRole(const FName& Role) const
+{
+	if (const FContextualAnimTrack* Result = FindTrack(Role))
+	{
+		return Result->Settings.PreviewActorClass;
+	}
+
+	return nullptr;
+}
+
+EContextualAnimJoinRule UContextualAnimSceneAsset::GetJoinRuleForRole(const FName& Role) const
+{
+	const FContextualAnimTrack* Result = FindTrack(Role);
+	return Result ? Result->Settings.JoinRule : EContextualAnimJoinRule::Default;
+}
diff --git a/Engine/Plugins/Experimental/Animation/ContextualAnimation/Source/ContextualAnimation/Private/ContextualAnimSceneAssetBase.cpp b/Engine/Plugins/Experimental/Animation/ContextualAnimation/Source/ContextualAnimation/Private/ContextualAnimSceneAssetBase.cpp
new file mode 100644
--- /dev/null
+++ b/Engine/Plugins/Experimental/Animation/ContextualAnimation/Source/ContextualAnimation/Private/ContextualAnimSceneAssetBase.cpp
@@ -0,0 +1,113 @@
+// Copyright Epic Games, Inc. All Rights Reserved.
+
+#include "ContextualAnimSceneAssetBase.h"
+#include "AnimationRuntime.h"
+#include "Animation/AnimMontage.h"
+#include "Animation/AnimationPoseData.h"
+#include "Animation/AnimTypes.h"
+#include "Containers/ArrayView.h"
+#include "ContextualAnimScenePivotProvider.h"
+#include "ContextualAnimUtilities.h"
+#include "ContextualAnimMetadata.h"
+
+UContextualAnimSceneAssetBase::UContextualAnimSceneAssetBase(const FObjectInitializer& ObjectInitializer)
+	: Super(ObjectInitializer)
+{
+	SampleRate = 15;
+	MeshToComponent = FTransform(FRotator(0.f, -90.f, 0.f));
+}
+
+FTransform UContextualAnimSceneAssetBase::ExtractTransformFromAnimData(const FContextualAnimData& AnimData, float Time) const
+{
+	const FTransform MeshToComponentInverse = MeshToComponent.Inverse();
+
+	if (const UAnimMontage* Animation = AnimData.Animation)
+	{
+		FCSPose<FCompactPose> ComponentSpacePose;
+		UContextualAnimUtilities::ExtractComponentSpacePose(Animation, FBoneContainer({ 0 }, FCurveEvaluationOption(false), *Animation->GetSkeleton()), Time, false, ComponentSpacePose);
+		return MeshToComponentInverse * (ComponentSpacePose.GetComponentSpaceTransform(FCompactPoseBoneIndex(0)) * AnimData.MeshToScene);
+	}
+	else
+	{
+		return MeshToComponentInverse * AnimData.MeshToScene;
+	}
+}
+
+void UContextualAnimSceneAssetBase::GenerateAlignmentTracksRelativeToScenePivot(FContextualAnimData& AnimData)
+{
+	const FTransform MeshToComponentInverse = MeshToComponent.Inverse();
+	const float SampleInterval = 1.f / SampleRate;
+
+	const int32 TotalTracks = AlignmentSections.Num();
+	AnimData.AlignmentData.Tracks.AnimationTracks.Reset(TotalTracks);
+	AnimData.AlignmentData.Tracks.TrackNames.Reset(TotalTracks);
+	AnimData.AlignmentData.SampleInterval = SampleInterval;
+
+	const UAnimMontage* Animation = AnimData.Animation;
+	if (Animation)
+	{
+		TArray<FBoneIndexType> RequiredBoneIndexArray;
+		RequiredBoneIndexArray.Add(0);
+		FBoneContainer BoneContainer = FBoneContainer(RequiredBoneIndexArray, FCurveEvaluationOption(false), *Animation->GetSkeleton());
+
+		for (int32 Idx = 0; Idx < AlignmentSections.Num(); Idx++)
+		{
+			float Time = 0.f;
+			float EndTime = Animation->GetPlayLength();
+			int32 SampleIndex = 0;
+
+			AnimData.AlignmentData.Tracks.TrackNames.Add(AlignmentSections[Idx].SectionName);
+			FRawAnimSequenceTrack& SceneTrack = AnimData.AlignmentData.Tracks.AnimationTracks.AddZeroed_GetRef();
+			while (Time < EndTime)
+			{
+				Time = FMath::Clamp(SampleIndex * SampleInterval, 0.f, EndTime);
+				SampleIndex++;
+
+				FCSPose<FCompactPose> ComponentSpacePose;
+				UContextualAnimUtilities::ExtractComponentSpacePose(Animation, BoneContainer, Time, false, ComponentSpacePose);
+
+				const FTransform RootTransform = MeshToComponentInverse * (ComponentSpacePose.GetComponentSpaceTransform(FCompactPoseBoneIndex(0)) * AnimData.MeshToScene);
+				const FTransform ScenePivotTransform = AlignmentSections[Idx].ScenePivot;
+				const FTransform RootRelativeToScenePivot = RootTransform.GetRelativeTransform(ScenePivotTransform);
+
+				SceneTrack.PosKeys.Add(RootRelativeToScenePivot.GetLocation());
+				SceneTrack.RotKeys.Add(RootRelativeToScenePivot.GetRotation());
+				SceneTrack.ScaleKeys.Add(RootRelativeToScenePivot.GetScale3D());
+			}
+		}
+	}
+	else
+	{
+		for (int32 Idx = 0; Idx < AlignmentSections.Num(); Idx++)
+		{
+			AnimData.AlignmentData.Tracks.TrackNames.Add(AlignmentSections[Idx].SectionName);
+			FRawAnimSequenceTrack& SceneTrack = AnimData.AlignmentData.Tracks.AnimationTracks.AddZeroed_GetRef();
+
+			const FTransform RootTransform = MeshToComponentInverse * AnimData.MeshToScene;
+			const FTransform ScenePivotTransform = AlignmentSections[Idx].ScenePivot;
+			const FTransform RootRelativeToScenePivot = RootTransform.GetRelativeTransform(ScenePivotTransform);
+
+			SceneTrack.PosKeys.Add(RootRelativeToScenePivot.GetLocation());
+			SceneTrack.RotKeys.Add(RootRelativeToScenePivot.GetRotation());
+			SceneTrack.ScaleKeys.Add(RootRelativeToScenePivot.GetScale3D());
+		}
+	}
+}
+
+void UContextualAnimSceneAssetBase::PreSave(const class ITargetPlatform* TargetPlatform)
+{
+	Super::PreSave(TargetPlatform);
+
+	// Generate scene pivot for each alignment section
+	for (int32 Idx = 0; Idx < AlignmentSections.Num(); Idx++)
+	{
+		if (AlignmentSections[Idx].ScenePivotProvider)
+		{
+			AlignmentSections[Idx].ScenePivot = AlignmentSections[Idx].ScenePivotProvider->CalculateScenePivot_Source();
+		}
+		else
+		{
+			AlignmentSections[Idx].ScenePivot = MeshToComponent.Inverse();
+		}
+	}
+}
diff --git a/Engine/Plugins/Experimental/Animation/ContextualAnimation/Source/ContextualAnimation/Private/ContextualAnimSceneInstance.cpp b/Engine/Plugins/Experimental/Animation/ContextualAnimation/Source/ContextualAnimation/Private/ContextualAnimSceneInstance.cpp
new file mode 100644
--- /dev/null
+++ b/Engine/Plugins/Experimental/Animation/ContextualAnimation/Source/ContextualAnimation/Private/ContextualAnimSceneInstance.cpp
@@ -0,0 +1,311 @@
+// Copyright Epic Games, Inc. All Rights Reserved.
+
+#include "ContextualAnimSceneInstance.h"
+#include "Animation/AnimMontage.h"
+#include "Animation/AnimInstance.h"
+#include "Components/SkeletalMeshComponent.h"
+#include "GameFramework/Character.h"
+#include "GameFramework/CharacterMovementComponent.h"
+#include "MotionWarpingComponent.h"
+#include "ContextualAnimActorInterface.h"
+#include "ContextualAnimScenePivotProvider.h"
+#include "ContextualAnimSceneAsset.h"
+#include "ContextualAnimSceneActorComponent.h"
+#include "Engine/World.h"
+#include "DrawDebugHelpers.h"
+
+//@TODO: Move to ContextualAnimUtilities
+static UAnimInstance* TryGetAnimInstance(AActor* Actor)
+{
+	UAnimInstance* AnimInstance = nullptr;
+	if (Actor)
+	{
+		if (ACharacter* Character = Cast<ACharacter>(Actor))
+		{
+			USkeletalMeshComponent* SkelMeshComp = Character->GetMesh();
+			AnimInstance = SkelMeshComp ? SkelMeshComp->GetAnimInstance() : nullptr;
+		}
+		else if (Actor->GetClass()->ImplementsInterface(UContextualAnimActorInterface::StaticClass()))
+		{
+			USkeletalMeshComponent* SkelMeshComp = IContextualAnimActorInterface::Execute_GetMesh(Actor);
+			AnimInstance = SkelMeshComp ? SkelMeshComp->GetAnimInstance() : nullptr;
+		}
+		else if (USkeletalMeshComponent* SkelMeshComp = Actor->FindComponentByClass<USkeletalMeshComponent>())
+		{
+			AnimInstance = SkelMeshComp->GetAnimInstance();
+		}
+	}
+
+	return AnimInstance;
+}
+
+//================================================================================================================
+
+FTransform FContextualAnimSceneActorData::GetTransform() const
+{
+	FTransform Transform = FTransform::Identity;
+
+	if(Actor.IsValid())
+	{
+		//@TODO: Cache this during the binding
+		if (const UContextualAnimSceneActorComponent* Comp = Actor->FindComponentByClass<UContextualAnimSceneActorComponent>())
+		{
+			Transform = Comp->GetComponentTransform();
+		}
+		else
+		{
+			Transform = Actor->GetActorTransform();
+		}
+	}
+
+	return Transform;
+}
+
+float FContextualAnimSceneActorData::GetAnimTime() const
+{
+	if(const UAnimMontage* Animation = AnimDataPtr->Animation)
+	{
+		if (UAnimInstance* AnimInstance = TryGetAnimInstance(Actor.Get()))
+		{
+			if(const FAnimMontageInstance* MontageInstance = AnimInstance->GetActiveInstanceForMontage(Animation))
+			{
+				return MontageInstance->GetPosition();
+			}
+		}
+	}
+	
+	return 0.f;
+}
+
+//================================================================================================================
+
+UContextualAnimSceneInstance::UContextualAnimSceneInstance(const FObjectInitializer& ObjectInitializer)
+	: Super(ObjectInitializer)
+{
+
+}
+
+UWorld* UContextualAnimSceneInstance::GetWorld()const
+{
+	return CastChecked<UWorld>(GetOuter()->GetWorld());
+}
+
+void UContextualAnimSceneInstance::Tick(float DeltaTime)
+{
+}
+
+bool UContextualAnimSceneInstance::IsActorInThisScene(const AActor* Actor) const
+{
+	if(Actor)
+	{
+		for (const auto& Pair : SceneActorMap)
+		{
+			const FContextualAnimSceneActorData& Data = Pair.Value;
+			if(Data.GetActor() == Actor)
+			{
+				return true;
+			}
+		}
+	}
+
+	return false;
+}
+
+void UContextualAnimSceneInstance::Join(FContextualAnimSceneActorData& Data)
+{
+	AActor* Actor = Data.GetActor();
+	if (UAnimInstance* AnimInstance = TryGetAnimInstance(Actor))
+	{
+		if (UMotionWarpingComponent* MotionWarpingComp = Actor->FindComponentByClass<UMotionWarpingComponent>())
+		{
+			for (const auto& Pair : AlignmentSectionToScenePivotList)
+			{
+				const FName SyncPointName = Pair.Key;
+				const float SyncTime = Data.AnimDataPtr->GetSyncTimeForWarpSection(SyncPointName);
+				const FTransform AlignmentTransform = (Data.AnimDataPtr->AlignmentData.ExtractTransformAtTime(SyncPointName, SyncTime) * Pair.Value);
+				MotionWarpingComp->AddOrUpdateSyncPoint(SyncPointName, AlignmentTransform);
+			}
+		}
+
+		UAnimMontage* Animation = Data.AnimDataPtr->Animation;
+		AnimInstance->Montage_Play(Animation, 1.f, EMontagePlayReturnType::MontageLength, Data.AnimStartTime);
+
+		AnimInstance->OnPlayMontageNotifyBegin.AddUniqueDynamic(this, &UContextualAnimSceneInstance::OnNotifyBeginReceived);
+		AnimInstance->OnPlayMontageNotifyEnd.AddUniqueDynamic(this, &UContextualAnimSceneInstance::OnNotifyEndReceived);
+		AnimInstance->OnMontageBlendingOut.AddUniqueDynamic(this, &UContextualAnimSceneInstance::OnMontageBlendingOut);
+
+		if (UCharacterMovementComponent* CharacterMovementComp = Actor->FindComponentByClass<UCharacterMovementComponent>())
+		{
+			CharacterMovementComp->SetMovementMode(MOVE_Flying);
+		}
+	}
+
+	SetIgnoreCollisionWithOtherActors(Actor, true);
+
+	OnActorJoined.ExecuteIfBound(this, Actor);
+}
+
+void UContextualAnimSceneInstance::Leave(FContextualAnimSceneActorData& Data)
+{
+	AActor* Actor = Data.GetActor();
+	if (UAnimInstance* AnimInstance = TryGetAnimInstance(Actor))
+	{
+		UAnimMontage* CurrentMontage = AnimInstance->GetCurrentActiveMontage();
+		check(CurrentMontage);
+
+		// Check if we have an exit section and transition to it, otherwise just stop the montage
+		// @TODO: This is temp until we add a solid way to deal with different states
+		static const FName ExitSectionName = FName(TEXT("Exit"));
+		const int32 SectionIdx = CurrentMontage->GetSectionIndex(ExitSectionName);
+		if (SectionIdx != INDEX_NONE)
+		{
+			// Unbind blend out delegate for a moment so we don't get it during the transition
+			AnimInstance->OnMontageBlendingOut.RemoveDynamic(this, &UContextualAnimSceneInstance::OnMontageBlendingOut);
+
+			AnimInstance->Montage_Play(CurrentMontage, 1.f);
+			AnimInstance->Montage_JumpToSection(ExitSectionName, CurrentMontage);
+
+			AnimInstance->OnMontageBlendingOut.AddUniqueDynamic(this, &UContextualAnimSceneInstance::OnMontageBlendingOut);
+		}
+		else
+		{
+			AnimInstance->Montage_Stop(CurrentMontage->BlendOut.GetBlendTime(), CurrentMontage);
+		}
+	}
+}
+
+void UContextualAnimSceneInstance::Start()
+{
+	const int32 AlignmentSectionsNum = SceneAsset->AlignmentSections.Num();
+	
+	AlignmentSectionToScenePivotList.Reset(AlignmentSectionsNum);
+
+	for (int32 Idx = 0; Idx < AlignmentSectionsNum; Idx++)
+	{
+		FTransform ScenePivot = FTransform::Identity;
+
+		if(SceneAsset->AlignmentSections[Idx].ScenePivotProvider)
+		{
+			ScenePivot = SceneAsset->AlignmentSections[Idx].ScenePivotProvider->CalculateScenePivot_Runtime(SceneActorMap);
+		}
+		else
+		{
+			ScenePivot = SceneActorMap.FindRef(SceneAsset->PrimaryRole).GetTransform();
+		}
+
+		AlignmentSectionToScenePivotList.Add(MakeTuple(SceneAsset->AlignmentSections[Idx].SectionName, ScenePivot));
+
+		//DrawDebugCoordinateSystem(GetWorld(), SceneOrigin.GetLocation(), SceneOrigin.Rotator(), 100.f, false, 10.f, 0, 5.f);
+	}
+
+	for (auto& Pair : SceneActorMap)
+	{
+		const FName& Role = Pair.Key;
+		FContextualAnimSceneActorData& Data = Pair.Value;
+
+		if (SceneAsset->GetJoinRuleForRole(Role) == EContextualAnimJoinRule::Default)
+		{
+			Join(Data);
+		}
+	}
+}
+
+void UContextualAnimSceneInstance::Stop()
+{
+	for (auto& Pair : SceneActorMap)
+	{
+		Leave(Pair.Value);
+	}
+}
+
+void UContextualAnimSceneInstance::OnMontageBlendingOut(UAnimMontage* Montage, bool bInterrupted)
+{
+	UE_LOG(LogContextualAnim, Log, TEXT("UContextualAnimSceneInstance::OnMontageBlendingOut Montage: %s"), *GetNameSafe(Montage));
+
+	for (auto& Pair : SceneActorMap)
+	{
+		if (Pair.Value.AnimDataPtr->Animation == Montage)
+		{
+			AActor* Actor = Pair.Value.GetActor();
+			if (UAnimInstance* AnimInstance = TryGetAnimInstance(Actor))
+			{
+				AnimInstance->OnPlayMontageNotifyBegin.RemoveDynamic(this, &UContextualAnimSceneInstance::OnNotifyBeginReceived);
+				AnimInstance->OnPlayMontageNotifyEnd.RemoveDynamic(this, &UContextualAnimSceneInstance::OnNotifyEndReceived);
+				AnimInstance->OnMontageBlendingOut.RemoveDynamic(this, &UContextualAnimSceneInstance::OnMontageBlendingOut);
+
+				if (UCharacterMovementComponent* CharacterMovementComp = Actor->FindComponentByClass<UCharacterMovementComponent>())
+				{
+					CharacterMovementComp->SetMovementMode(MOVE_Walking);
+				}
+			}
+
+			SetIgnoreCollisionWithOtherActors(Pair.Value.GetActor(), false);
+
+			OnActorLeft.ExecuteIfBound(this, Actor);
+
+			break;
+		}
+	}
+
+	bool bShouldEnd = true;
+	for (auto& Pair : SceneActorMap)
+	{
+		const FContextualAnimSceneActorData& Data = Pair.Value;
+		if (UAnimInstance* AnimInstance = TryGetAnimInstance(Data.GetActor()))
+		{
+			if(AnimInstance->Montage_IsPlaying(Data.AnimDataPtr->Animation))
+			{
+				bShouldEnd = false;
+				break;
+			}
+		}
+	}
+
+	if(bShouldEnd)
+	{
+		OnSceneEnded.ExecuteIfBound(this);
+	}
+}
+
+void UContextualAnimSceneInstance::OnNotifyBeginReceived(FName NotifyName, const FBranchingPointNotifyPayload& BranchingPointNotifyPayload)
+{
+	UE_LOG(LogContextualAnim, Log, TEXT("UContextualAnimSceneInstance::OnNotifyBeginReceived NotifyName: %s Montage: %s"),
+	*NotifyName.ToString(), *GetNameSafe(BranchingPointNotifyPayload.SequenceAsset));
+
+	// @TODO: For now just use a hard-coded name to identify the event. We should change this in the future
+	static const FName LateJoinNotifyName = FName(TEXT("ContextualAnimLateJoin"));
+	if(NotifyName == LateJoinNotifyName)
+	{
+		for (auto& Pair : SceneActorMap)
+		{
+			const FName& Role = Pair.Key;
+			FContextualAnimSceneActorData& Data = Pair.Value;
+
+			if (SceneAsset->GetJoinRuleForRole(Role) == EContextualAnimJoinRule::Late)
+			{
+				Join(Data);
+			}
+		}
+	}
+}
+
+void UContextualAnimSceneInstance::OnNotifyEndReceived(FName NotifyName, const FBranchingPointNotifyPayload& BranchingPointNotifyPayload)
+{
+	UE_LOG(LogContextualAnim, Log, TEXT("UContextualAnimSceneInstance::OnNotifyEndReceived NotifyName: %s Montage: %s"),
+	*NotifyName.ToString(), *GetNameSafe(BranchingPointNotifyPayload.SequenceAsset));
+}
+
+void UContextualAnimSceneInstance::SetIgnoreCollisionWithOtherActors(AActor* Actor, bool bValue) const
+{
+	for (auto& Pair : SceneActorMap)
+	{
+		AActor* OtherActor = Pair.Value.GetActor();
+		if(OtherActor != Actor)
+		{
+			if (UPrimitiveComponent* RootPrimitiveComponent = Cast<UPrimitiveComponent>(Actor->GetRootComponent()))
+			{
+				RootPrimitiveComponent->IgnoreActorWhenMoving(OtherActor, bValue);
+			}
+		}
+	}
+}
\ No newline at end of file
diff --git a/Engine/Plugins/Experimental/Animation/ContextualAnimation/Source/ContextualAnimation/Private/ContextualAnimScenePivotProvider.cpp b/Engine/Plugins/Experimental/Animation/ContextualAnimation/Source/ContextualAnimation/Private/ContextualAnimScenePivotProvider.cpp
new file mode 100644
--- /dev/null
+++ b/Engine/Plugins/Experimental/Animation/ContextualAnimation/Source/ContextualAnimation/Private/ContextualAnimScenePivotProvider.cpp
@@ -0,0 +1,108 @@
+// Copyright Epic Games, Inc. All Rights Reserved.
+
+#include "ContextualAnimScenePivotProvider.h"
+#include "ContextualAnimSceneAsset.h"
+#include "ContextualAnimUtilities.h"
+
+// UContextualAnimScenePivotProvider
+//==================================================
+
+UContextualAnimScenePivotProvider::UContextualAnimScenePivotProvider(const FObjectInitializer& ObjectInitializer)
+	: Super(ObjectInitializer)
+{
+}
+
+const UContextualAnimSceneAssetBase* UContextualAnimScenePivotProvider::GetSceneAsset() const
+{
+	return CastChecked<UContextualAnimSceneAssetBase>(GetOuter());
+}
+
+// UContextualAnimScenePivotProvider_Default
+//==================================================
+
+UContextualAnimScenePivotProvider_Default::UContextualAnimScenePivotProvider_Default(const FObjectInitializer& ObjectInitializer)
+	: Super(ObjectInitializer)
+{
+}
+
+FTransform UContextualAnimScenePivotProvider_Default::CalculateScenePivot_Source() const
+{
+	FTransform ScenePivot = FTransform::Identity;
+
+	if(const UContextualAnimSceneAsset* SceneAsset = Cast<UContextualAnimSceneAsset>(GetSceneAsset()))
+	{
+		if(const FContextualAnimTrack* PrimaryTrack = SceneAsset->DataContainer.Find(PrimaryRole))
+		{
+			if(const FContextualAnimTrack* SecondaryTrack = SceneAsset->DataContainer.Find(SecondaryRole))
+			{
+				const FTransform PrimaryTransform = SceneAsset->ExtractTransformFromAnimData(PrimaryTrack->AnimData, 0.f);
+				const FTransform SecondaryTransform = SceneAsset->ExtractTransformFromAnimData(SecondaryTrack->AnimData, 0.f);
+
+				ScenePivot.SetLocation(FMath::Lerp<FVector>(PrimaryTransform.GetLocation(), SecondaryTransform.GetLocation(), Weight));
+				ScenePivot.SetRotation((SecondaryTransform.GetLocation() - PrimaryTransform.GetLocation()).GetSafeNormal2D().ToOrientationQuat());
+			}
+		}
+	}
+
+	return ScenePivot;
+}
+
+FTransform UContextualAnimScenePivotProvider_Default::CalculateScenePivot_Runtime(const TMap<FName, FContextualAnimSceneActorData>& SceneActorMap) const
+{
+	FTransform ScenePivot = FTransform::Identity;
+
+	if(const UContextualAnimSceneAsset* SceneAsset = Cast<UContextualAnimSceneAsset>(GetSceneAsset()))
+	{
+		if(const AActor* PrimaryActor = SceneActorMap.Find(PrimaryRole)->GetActor())
+		{
+			if (const AActor* SecondaryActor = SceneActorMap.Find(SecondaryRole)->GetActor())
+			{
+				ScenePivot.SetLocation(FMath::Lerp<FVector>(PrimaryActor->GetActorLocation(), SecondaryActor->GetActorLocation(), Weight));
+				ScenePivot.SetRotation((SecondaryActor->GetActorLocation() - PrimaryActor->GetActorLocation()).GetSafeNormal2D().ToOrientationQuat());
+			}
+		}
+	}
+
+	return ScenePivot;
+}
+
+// UContextualAnimScenePivotProvider_RelativeTo
+//==================================================
+
+UContextualAnimScenePivotProvider_RelativeTo::UContextualAnimScenePivotProvider_RelativeTo(const FObjectInitializer& ObjectInitializer)
+	: Super(ObjectInitializer)
+{
+}
+
+FTransform UContextualAnimScenePivotProvider_RelativeTo::CalculateScenePivot_Source() const
+{
+	FTransform ScenePivot = FTransform::Identity;
+
+	if (const UContextualAnimSceneAsset* SceneAsset = Cast<UContextualAnimSceneAsset>(GetSceneAsset()))
+	{
+		if(const FContextualAnimTrack* RelativeToTrack = SceneAsset->DataContainer.Find(RelativeToRole))
+		{
+			ScenePivot = SceneAsset->ExtractTransformFromAnimData(RelativeToTrack->AnimData, 0.f);
+		}
+	}
+
+	return ScenePivot;
+}
+
+FTransform UContextualAnimScenePivotProvider_RelativeTo::CalculateScenePivot_Runtime(const TMap<FName, FContextualAnimSceneActorData>& SceneActorMap) const
+{
+	FTransform ScenePivot = FTransform::Identity;
+
+	if (const UContextualAnimSceneAsset* SceneAsset = Cast<UContextualAnimSceneAsset>(GetSceneAsset()))
+	{
+		if(const FContextualAnimSceneActorData* Data = SceneActorMap.Find(RelativeToRole))
+		{
+			if(const AActor* Actor = Data->GetActor())
+			{
+				ScenePivot = Actor->GetActorTransform();
+			}
+		}
+	}
+
+	return ScenePivot;
+}
\ No newline at end of file
diff --git a/Engine/Plugins/Experimental/Animation/ContextualAnimation/Source/ContextualAnimation/Private/ContextualAnimTypes.cpp b/Engine/Plugins/Experimental/Animation/ContextualAnimation/Source/ContextualAnimation/Private/ContextualAnimTypes.cpp
new file mode 100644
--- /dev/null
+++ b/Engine/Plugins/Experimental/Animation/ContextualAnimation/Source/ContextualAnimation/Private/ContextualAnimTypes.cpp
@@ -0,0 +1,170 @@
+// Copyright Epic Games, Inc. All Rights Reserved.
+
+#include "ContextualAnimTypes.h"
+#include "AnimationUtils.h"
+#include "Animation/AnimMontage.h"
+#include "Containers/ArrayView.h"
+#include "AnimNotifyState_MotionWarping.h"
+#include "RootMotionModifier.h"
+
+DEFINE_LOG_CATEGORY(LogContextualAnim);
+
+const FContextualAnimData FContextualAnimData::EmptyAnimData = FContextualAnimData();
+
+// FContextualAnimAlignmentTrackContainer
+///////////////////////////////////////////////////////////////////////
+
+FTransform FContextualAnimAlignmentTrackContainer::ExtractTransformAtTime(const FName& TrackName, float Time) const
+{
+	const int32 TrackIndex = Tracks.TrackNames.IndexOfByKey(TrackName);
+	return ExtractTransformAtTime(TrackIndex, Time);
+}
+
+FTransform FContextualAnimAlignmentTrackContainer::ExtractTransformAtTime(int32 TrackIndex, float Time) const
+{
+	FTransform AlignmentTransform = FTransform::Identity;
+
+	if (Tracks.AnimationTracks.IsValidIndex(TrackIndex))
+	{
+		const FRawAnimSequenceTrack& Track = Tracks.AnimationTracks[TrackIndex];
+		const int32 TotalFrames = Track.PosKeys.Num();
+		const float TrackLength = (TotalFrames - 1) * SampleInterval;
+		FAnimationUtils::ExtractTransformFromTrack(Time, TotalFrames, TrackLength, Track, EAnimInterpolationType::Linear, AlignmentTransform);
+	}
+
+	return AlignmentTransform;
+}
+
+float FContextualAnimData::GetSyncTimeForWarpSection(int32 WarpSectionIndex) const
+{
+	//@TODO: We need a better way to identify warping sections withing the animation. This is just a temp solution
+	//@TODO: We should cache this data
+
+	float Result = 0.f;
+
+	if(Animation && WarpSectionIndex >= 0)
+	{
+		FName LastSyncPointName = NAME_None;
+		int32 LastWarpSectionIndex = INDEX_NONE;
+
+		for (int32 Idx = 0; Idx < Animation->Notifies.Num(); Idx++)
+		{
+			const FAnimNotifyEvent& NotifyEvent = Animation->Notifies[Idx];
+			if(const UAnimNotifyState_MotionWarping* Notify = Cast<const UAnimNotifyState_MotionWarping>(NotifyEvent.NotifyStateClass))
+			{
+				//@TODO: We may want to add a function in UAnimNotifyState_MotionWarping that returns the SyncPointName?
+				if(const URootMotionModifierConfig_Warp* Config = Cast<const URootMotionModifierConfig_Warp>(Notify->RootMotionModifierConfig))
+				{
+					const FName SyncPointName = Config->SyncPointName;
+					if(SyncPointName != NAME_None)
+					{
+						// First valid warping window. Initialize everything
+						if (LastWarpSectionIndex == INDEX_NONE)
+						{
+							LastSyncPointName = SyncPointName;
+							Result = NotifyEvent.GetEndTriggerTime();
+							LastWarpSectionIndex = 0;
+						}
+						// If we hit another warping window but the sync point is the same as the previous, update SyncTime.
+						// This is to deal with cases where a first short window is used to face the alignment point and a second one to perform the rest of the warp
+						else if(SyncPointName == LastSyncPointName)
+						{
+							Result = NotifyEvent.GetEndTriggerTime();
+						}
+						// If we hit another warping window but with a different sync point name means that we have hit the first window of another warping section
+						else
+						{
+							// If we haven't reached the desired WarpSection yet. Update control vars and keep moving
+							if(WarpSectionIndex > LastWarpSectionIndex)
+							{
+								LastSyncPointName = SyncPointName;
+								Result = NotifyEvent.GetEndTriggerTime();
+								LastWarpSectionIndex++;
+							}
+							// Otherwise, stop here and return the value of the last window we found
+							else 
+							{
+								break;
+							}
+						}
+					}
+				}
+			}
+		}
+	}
+
+	return Result;
+}
+
+float FContextualAnimData::GetSyncTimeForWarpSection(const FName& WarpSectionName) const
+{
+	//@TODO: We need a better way to identify warping sections withing the animation. This is just a temp solution
+	//@TODO: We should cache this data
+
+	float Result = 0.f;
+
+	if (Animation && WarpSectionName != NAME_None)
+	{
+		for (int32 Idx = 0; Idx < Animation->Notifies.Num(); Idx++)
+		{
+			const FAnimNotifyEvent& NotifyEvent = Animation->Notifies[Idx];
+			if (const UAnimNotifyState_MotionWarping* Notify = Cast<const UAnimNotifyState_MotionWarping>(NotifyEvent.NotifyStateClass))
+			{
+				//@TODO: We may want to add a function in UAnimNotifyState_MotionWarping that returns the SyncPointName?
+				if (const URootMotionModifierConfig_Warp* Config = Cast<const URootMotionModifierConfig_Warp>(Notify->RootMotionModifierConfig))
+				{
+					const FName SyncPointName = Config->SyncPointName;
+					if (WarpSectionName == SyncPointName)
+					{
+						const float NotifyEndTriggerTime = NotifyEvent.GetEndTriggerTime();
+						if(NotifyEndTriggerTime > Result)
+						{
+							Result = NotifyEndTriggerTime;
+						}
+					}
+				}
+			}
+		}
+	}
+
+	return Result;
+}
+
+float FContextualAnimData::FindBestAnimStartTime(const FVector& LocalLocation) const
+{
+	float BestTime = 0.f;
+
+	if (AnimMaxStartTime < 0.f)
+	{
+		return BestTime;
+	}
+
+	const FVector SyncPointLocation = GetAlignmentTransformAtSyncTime().GetLocation();
+	const float PerfectDistToSyncPointSq = GetAlignmentTransformAtEntryTime().GetTranslation().SizeSquared2D();
+	const float ActualDistToSyncPointSq = FVector::DistSquared2D(LocalLocation, SyncPointLocation);
+
+	if (ActualDistToSyncPointSq < PerfectDistToSyncPointSq)
+	{
+		float BestDistance = MAX_FLT;
+		TArrayView<const FVector> PosKeys(AlignmentData.Tracks.AnimationTracks[0].PosKeys.GetData(), AlignmentData.Tracks.AnimationTracks[0].PosKeys.Num());
+
+		//@TODO: Very simple search for now. Replace with Distance Matching + Pose Matching
+		for (int32 Idx = 0; Idx < PosKeys.Num(); Idx++)
+		{
+			const float Time = Idx * AlignmentData.SampleInterval;
+			if (AnimMaxStartTime > 0.f && Time >= AnimMaxStartTime)
+			{
+				break;
+			}
+
+			const float DistFromCurrentFrameToSyncPointSq = FVector::DistSquared2D(SyncPointLocation, PosKeys[Idx]);
+			if (DistFromCurrentFrameToSyncPointSq < ActualDistToSyncPointSq)
+			{
+				BestTime = Time;
+				break;
+			}
+		}
+	}
+
+	return BestTime;
+}
\ No newline at end of file
diff --git a/Engine/Plugins/Experimental/Animation/ContextualAnimation/Source/ContextualAnimation/Private/ContextualAnimUtilities.cpp b/Engine/Plugins/Experimental/Animation/ContextualAnimation/Source/ContextualAnimation/Private/ContextualAnimUtilities.cpp
new file mode 100644
--- /dev/null
+++ b/Engine/Plugins/Experimental/Animation/ContextualAnimation/Source/ContextualAnimation/Private/ContextualAnimUtilities.cpp
@@ -0,0 +1,120 @@
+// Copyright Epic Games, Inc. All Rights Reserved.
+
+#include "ContextualAnimUtilities.h"
+#include "AnimationUtils.h"
+#include "AnimationRuntime.h"
+#include "Animation/AnimMontage.h"
+#include "Animation/AnimationPoseData.h"
+#include "Animation/CustomAttributesRuntime.h"
+#include "Animation/AnimTypes.h"
+#include "DrawDebugHelpers.h"
+#include "Engine/World.h"
+#include "ContextualAnimTypes.h"
+#include "ContextualAnimSceneAsset.h"
+
+void UContextualAnimUtilities::ExtractLocalSpacePose(const UAnimSequenceBase* Animation, const FBoneContainer& BoneContainer, float Time, bool bExtractRootMotion, FCompactPose& OutPose)
+{
+	OutPose.SetBoneContainer(&BoneContainer);
+
+	FBlendedCurve Curve;
+	Curve.InitFrom(BoneContainer);
+
+	FAnimExtractContext Context(Time, bExtractRootMotion);
+
+	FStackCustomAttributes Attributes;
+	FAnimationPoseData AnimationPoseData(OutPose, Curve, Attributes);
+	if (const UAnimSequence* AnimSequence = Cast<UAnimSequence>(Animation))
+	{
+		AnimSequence->GetBonePose(AnimationPoseData, Context);
+	}
+	else if (const UAnimMontage* AnimMontage = Cast<UAnimMontage>(Animation))
+	{
+		const FAnimTrack& AnimTrack = AnimMontage->SlotAnimTracks[0].AnimTrack;
+		AnimTrack.GetAnimationPose(AnimationPoseData, Context);
+	}
+}
+
+void UContextualAnimUtilities::ExtractComponentSpacePose(const UAnimSequenceBase* Animation, const FBoneContainer& BoneContainer, float Time, bool bExtractRootMotion, FCSPose<FCompactPose>& OutPose)
+{
+	FCompactPose Pose;
+	ExtractLocalSpacePose(Animation, BoneContainer, Time, bExtractRootMotion, Pose);
+	OutPose.InitPose(MoveTemp(Pose));
+}
+
+FTransform UContextualAnimUtilities::ExtractRootMotionFromAnimation(const UAnimSequenceBase* Animation, float StartTime, float EndTime)
+{
+	if (const UAnimMontage* Anim = Cast<UAnimMontage>(Animation))
+	{
+		return Anim->ExtractRootMotionFromTrackRange(StartTime, EndTime);
+	}
+
+	if (const UAnimSequence* Anim = Cast<UAnimSequence>(Animation))
+	{
+		return Anim->ExtractRootMotionFromRange(StartTime, EndTime);
+	}
+
+	return FTransform::Identity;
+}
+
+void UContextualAnimUtilities::DrawDebugPose(const UWorld* World, const UAnimSequenceBase* Animation, float Time, const FTransform& LocalToWorldTransform, const FColor& Color, float LifeTime, float Thickness)
+{ 
+	if(World)
+	{
+		Time = FMath::Clamp(Time, 0.f, Animation->GetPlayLength());
+
+		const int32 TotalBones = Animation->GetSkeleton()->GetReferenceSkeleton().GetNum();
+		TArray<FBoneIndexType> RequiredBoneIndexArray;
+		RequiredBoneIndexArray.Reserve(TotalBones);
+		for (int32 Idx = 0; Idx < TotalBones; Idx++)
+		{
+			RequiredBoneIndexArray.Add(Idx);
+		}
+
+		FBoneContainer BoneContainer(RequiredBoneIndexArray, FCurveEvaluationOption(false), *Animation->GetSkeleton());
+		FCSPose<FCompactPose> ComponentSpacePose;
+		UContextualAnimUtilities::ExtractComponentSpacePose(Animation, BoneContainer, Time, true, ComponentSpacePose);
+
+		for (int32 Index = 0; Index < ComponentSpacePose.GetPose().GetNumBones(); ++Index)
+		{
+			const FCompactPoseBoneIndex CompactPoseBoneIndex = FCompactPoseBoneIndex(Index);
+			const FCompactPoseBoneIndex ParentIndex = ComponentSpacePose.GetPose().GetParentBoneIndex(CompactPoseBoneIndex);
+			FVector Start, End;
+
+			FLinearColor LineColor = FLinearColor::Red;
+			const FTransform Transform = ComponentSpacePose.GetComponentSpaceTransform(CompactPoseBoneIndex) * LocalToWorldTransform;
+
+			if (ParentIndex.GetInt() >= 0)
+			{
+				Start = (ComponentSpacePose.GetComponentSpaceTransform(ParentIndex) * LocalToWorldTransform).GetLocation();
+				End = Transform.GetLocation();
+			}
+			else
+			{
+				Start = LocalToWorldTransform.GetLocation();
+				End = Transform.GetLocation();
+			}
+
+			DrawDebugLine(World, Start, End, Color, false, LifeTime, 0, Thickness);
+		}
+	}
+}
+
+void UContextualAnimUtilities::DrawDebugScene(const UWorld* World, const UContextualAnimSceneAsset* SceneAsset, float Time, const FTransform& ToWorldTransform, const FColor& Color, float LifeTime, float Thickness)
+{
+	if (World && SceneAsset)
+	{
+		for(const auto& Pair : SceneAsset->DataContainer)
+		{
+			const FContextualAnimTrack& Track = Pair.Value;
+			const FTransform Transform = (SceneAsset->MeshToComponent * Track.AnimData.GetAlignmentTransformAtTime(Time)) * ToWorldTransform;
+			if(const UAnimMontage* Animation = Track.AnimData.Animation)
+			{
+				DrawDebugPose(World, Animation, Time, Transform, Color, LifeTime, Thickness);
+			}
+			else
+			{
+				DrawDebugCoordinateSystem(World, Transform.GetLocation(), Transform.Rotator(), 50.f, false, LifeTime, 0, Thickness);
+			}
+		}
+	}
+}
\ No newline at end of file
diff --git a/Engine/Plugins/Experimental/Animation/ContextualAnimation/Source/ContextualAnimation/Private/ContextualAnimation.cpp b/Engine/Plugins/Experimental/Animation/ContextualAnimation/Source/ContextualAnimation/Private/ContextualAnimation.cpp
--- a/Engine/Plugins/Experimental/Animation/ContextualAnimation/Source/ContextualAnimation/Private/ContextualAnimation.cpp
+++ b/Engine/Plugins/Experimental/Animation/ContextualAnimation/Source/ContextualAnimation/Private/ContextualAnimation.cpp
@@ -1,17 +1,42 @@
 // Copyright Epic Games, Inc. All Rights Reserved.
 
 #include "ContextualAnimation.h"
+#include "ContextualAnimManager.h"
 
 #define LOCTEXT_NAMESPACE "ContextualAnimationModule"
 
+TMap<const UWorld*, UContextualAnimManager*> FContextualAnimationModule::WorldToManagerMap;
+
 void FContextualAnimationModule::StartupModule()
 {
-
+	OnPreWorldInitDelegateHandle = FWorldDelegates::OnPreWorldInitialization.AddStatic(&FContextualAnimationModule::OnWorldInit);
+	OnPostWorldCleanupDelegateHandle = FWorldDelegates::OnPostWorldCleanup.AddStatic(&FContextualAnimationModule::OnWorldCleanup);
 }
 
 void FContextualAnimationModule::ShutdownModule()
 {
+	FWorldDelegates::OnPreWorldInitialization.Remove(OnPreWorldInitDelegateHandle);
+	FWorldDelegates::OnPostWorldCleanup.Remove(OnPostWorldCleanupDelegateHandle);
+
+	WorldToManagerMap.Empty();
+}
 
+void FContextualAnimationModule::AddReferencedObjects(FReferenceCollector& Collector)
+{
+	for (TPair<const UWorld*, UContextualAnimManager*>& Pair : WorldToManagerMap)
+	{
+		Collector.AddReferencedObject(Pair.Value, Pair.Key);
+	}
+}
+
+void FContextualAnimationModule::OnWorldInit(UWorld* World, const UWorld::InitializationValues IVS)
+{
+	WorldToManagerMap.Add(World, NewObject<UContextualAnimManager>(World));
+}
+
+void FContextualAnimationModule::OnWorldCleanup(UWorld* World, bool bSessionEnded, bool bCleanupResources)
+{
+	WorldToManagerMap.Remove(World);
 }
 
 #undef LOCTEXT_NAMESPACE
diff --git a/Engine/Plugins/Experimental/Animation/ContextualAnimation/Source/ContextualAnimation/Public/ContextualAnimActorInterface.h b/Engine/Plugins/Experimental/Animation/ContextualAnimation/Source/ContextualAnimation/Public/ContextualAnimActorInterface.h
new file mode 100644
--- /dev/null
+++ b/Engine/Plugins/Experimental/Animation/ContextualAnimation/Source/ContextualAnimation/Public/ContextualAnimActorInterface.h
@@ -0,0 +1,21 @@
+// Copyright Epic Games, Inc. All Rights Reserved.
+
+#pragma once
+
+#include "CoreMinimal.h"
+#include "UObject/Interface.h"
+#include "ContextualAnimActorInterface.generated.h"
+
+UINTERFACE()
+class UContextualAnimActorInterface : public UInterface
+{
+	GENERATED_UINTERFACE_BODY()
+};
+
+class IContextualAnimActorInterface : public IInterface
+{
+	GENERATED_IINTERFACE_BODY()
+
+	UFUNCTION(BlueprintNativeEvent, BlueprintCallable, Category = "Contextual Animation System")
+	class USkeletalMeshComponent* GetMesh() const;
+};
\ No newline at end of file
diff --git a/Engine/Plugins/Experimental/Animation/ContextualAnimation/Source/ContextualAnimation/Public/ContextualAnimAsset.h b/Engine/Plugins/Experimental/Animation/ContextualAnimation/Source/ContextualAnimation/Public/ContextualAnimAsset.h
--- a/Engine/Plugins/Experimental/Animation/ContextualAnimation/Source/ContextualAnimation/Public/ContextualAnimAsset.h
+++ b/Engine/Plugins/Experimental/Animation/ContextualAnimation/Source/ContextualAnimation/Public/ContextualAnimAsset.h
@@ -4,171 +4,13 @@
 
 #include "Engine/DataAsset.h"
 #include "Animation/AnimSequence.h"
+#include "ContextualAnimMetadata.h"
+#include "ContextualAnimTypes.h"
+#include "ContextualAnimUtilities.h"
+#include "ContextualAnimCompositeSceneAsset.h"
 #include "ContextualAnimAsset.generated.h"
 
-DECLARE_LOG_CATEGORY_EXTERN(LogContextualAnim, Log, All);
-
-UCLASS()
-class UContextualAnimUtilities : public UBlueprintFunctionLibrary
-{
-	GENERATED_BODY()
-
-public:
-
-	static void ExtractLocalSpacePose(const UAnimSequenceBase* Animation, const FBoneContainer& BoneContainer, float Time, bool bExtractRootMotion, FCompactPose& OutPose);
-	static void ExtractComponentSpacePose(const UAnimSequenceBase* Animation, const FBoneContainer& BoneContainer, float Time, bool bExtractRootMotion, FCSPose<FCompactPose>& OutPose);
-	static FTransform ExtractRootMotionFromAnimation(const UAnimSequenceBase* Animation, float StartTime, float EndTime);
-};
-
-USTRUCT()
-struct FAlignmentTrackContainer
-{
-	GENERATED_BODY()
-
-	UPROPERTY()
-	FRawAnimSequenceTrack Track;
-
-	UPROPERTY()
-	float SampleInterval;
-
-	FTransform ExtractTransformAtTime(float Time) const;
-
-	void Initialize(float InSampleInterval)
-	{
-		Track.PosKeys.Empty();
-		Track.RotKeys.Empty();
-		Track.ScaleKeys.Empty();
-
-		SampleInterval = InSampleInterval;
-	}
-};
-
-/** Stores the result of a query function */
-USTRUCT(BlueprintType)
-struct FContextualAnimQueryResult
-{
-	GENERATED_BODY()
-
-	UPROPERTY(BlueprintReadWrite, Category = "Defaults")
-	TSoftObjectPtr<UAnimMontage> Animation;
-
-	UPROPERTY(BlueprintReadWrite, Category = "Defaults")
-	FTransform EntryTransform;
-
-	UPROPERTY(BlueprintReadWrite, Category = "Defaults")
-	FTransform SyncTransform;
-
-	UPROPERTY(BlueprintReadWrite, Category = "Defaults")
-	float AnimStartTime = 0.f;
-
-	UPROPERTY(BlueprintReadWrite, Category = "Defaults")
-	int32 DataIndex = INDEX_NONE;
-
-	void Reset()
-	{
-		Animation = nullptr;
-		EntryTransform = SyncTransform = FTransform::Identity;
-		AnimStartTime = 0.f;
-		DataIndex = INDEX_NONE;
-	}
-
-	FORCEINLINE bool IsValid() const { return DataIndex != INDEX_NONE; }
-};
-
-/** Stores the parameters passed into query function */
-USTRUCT(BlueprintType)
-struct FContextualAnimQueryParams
-{
-	GENERATED_BODY()
-
-	UPROPERTY(EditAnywhere, BlueprintReadWrite, Category = "Defaults")
-	TWeakObjectPtr<const AActor> Querier;
-
-	UPROPERTY(EditAnywhere, BlueprintReadWrite, Category = "Defaults")
-	FTransform QueryTransform;
-
-	UPROPERTY(EditAnywhere, BlueprintReadWrite, Category = "Defaults")
-	bool bComplexQuery = false;
-
-	UPROPERTY(EditAnywhere, BlueprintReadWrite, Category = "Defaults")
-	bool bFindAnimStartTime = false;
-
-	FContextualAnimQueryParams(){}
-
-	FContextualAnimQueryParams(const AActor* InQuerier, bool bInComplexQuery, bool bInFindAnimStartTime)
-		: Querier(InQuerier), bComplexQuery(bInComplexQuery), bFindAnimStartTime(bInFindAnimStartTime) {}
-
-	FContextualAnimQueryParams(const FTransform& InQueryTransform, bool bInComplexQuery, bool bInFindAnimStartTime)
-		: QueryTransform(InQueryTransform), bComplexQuery(bInComplexQuery), bFindAnimStartTime(bInFindAnimStartTime) {}
-};
-
-USTRUCT(BlueprintType)
-struct FContextualAnimDistanceParam
-{
-	GENERATED_BODY()
-
-	UPROPERTY(VisibleAnywhere, BlueprintReadWrite, Category = "Defaults", meta = (UIMin = 0, ClampMin = 0))
-	float Value = 0.f;
-
-	UPROPERTY(EditAnywhere, BlueprintReadWrite, Category = "Defaults", meta = (UIMin = 0, ClampMin = 0))
-	float MinDistance = 0.f;
-
-	UPROPERTY(EditAnywhere, BlueprintReadWrite, Category = "Defaults", meta = (UIMin = 0, ClampMin = 0))
-	float MaxDistance = 0.f;
-};
-
-USTRUCT(BlueprintType)
-struct FContextualAnimAngleParam
-{
-	GENERATED_BODY()
-
-	UPROPERTY(VisibleAnywhere, BlueprintReadWrite, Category = "Defaults")
-	float Value = 0.f;
-
-	UPROPERTY(EditAnywhere, BlueprintReadWrite, Category = "Defaults")
-	float Tolerance = 0.f;
-};
-
-USTRUCT(BlueprintType)
-struct FContextualAnimData
-{
-	GENERATED_BODY()
-
-	UPROPERTY(EditAnywhere, BlueprintReadWrite, Category = "Defaults")
-	TSoftObjectPtr<UAnimMontage> Animation;
-
-	UPROPERTY(EditAnywhere, BlueprintReadWrite, Category = "Defaults", meta = (UIMin = 0, ClampMin = 0))
-	float EntryTime = 0.f;
-
-	UPROPERTY(EditAnywhere, BlueprintReadWrite, Category = "Defaults", meta = (UIMin = 0, ClampMin = 0))
-	float SyncTime = 0.f;
-
-	//@TODO: Temp solution to prevent FindBestAnimStartTime from returning a time that leaves us without enough room for warping
-	UPROPERTY(EditAnywhere, BlueprintReadWrite, Category = "Defaults", meta = (UIMin = 0, ClampMin = 0))
-	float AnimMinStartTime = 0.f;
-
-	UPROPERTY(EditAnywhere, BlueprintReadWrite, Category = "Defaults")
-	float OffsetFromOrigin = 0.f;
-
-	UPROPERTY(EditAnywhere, BlueprintReadWrite, Category = "Defaults")
-	FContextualAnimDistanceParam Distance;
-
-	UPROPERTY(EditAnywhere, BlueprintReadWrite, Category = "Defaults")
-	FContextualAnimAngleParam Angle;
-
-	UPROPERTY(EditAnywhere, BlueprintReadWrite, Category = "Defaults")
-	FContextualAnimAngleParam Facing;
-
-	UPROPERTY()
-	FAlignmentTrackContainer AlignmentData;
-
-	FORCEINLINE FTransform GetAlignmentTransformAtTime(float Time) const { return AlignmentData.ExtractTransformAtTime(Time); }
-	FORCEINLINE FTransform GetAlignmentTransformAtEntryTime() const { return AlignmentData.ExtractTransformAtTime(EntryTime); }
-	FORCEINLINE FTransform GetAlignmentTransformAtSyncTime() const  { return AlignmentData.ExtractTransformAtTime(SyncTime); }
-
-	float FindBestAnimStartTime(const FVector& LocalLocation) const;
-};
-
+/** Deprecated, it has been replaced by ContextualAnimCompositeSceneAsset. This file will be removed very soon */
 UCLASS(Blueprintable)
 class CONTEXTUALANIMATION_API UContextualAnimAsset : public UDataAsset
 {
@@ -193,7 +35,5 @@ public:
 
 	UContextualAnimAsset(const FObjectInitializer& ObjectInitializer);
 
-	virtual void PreSave(const class ITargetPlatform* TargetPlatform) override;
-
-	bool QueryData(FContextualAnimQueryResult& OutResult, const FContextualAnimQueryParams& QueryParams, const FTransform& ToWorldTransform) const;
+	bool QueryData(FContextualAnimQueryResult& OutResult, const FContextualAnimQueryParams& QueryParams, const FTransform& ToWorldTransform) const { return false; } 
 };
diff --git a/Engine/Plugins/Experimental/Animation/ContextualAnimation/Source/ContextualAnimation/Public/ContextualAnimComponent.h b/Engine/Plugins/Experimental/Animation/ContextualAnimation/Source/ContextualAnimation/Public/ContextualAnimComponent.h
deleted file mode 100644
diff --git a/Engine/Plugins/Experimental/Animation/ContextualAnimation/Source/ContextualAnimation/Public/ContextualAnimCompositeSceneAsset.h b/Engine/Plugins/Experimental/Animation/ContextualAnimation/Source/ContextualAnimation/Public/ContextualAnimCompositeSceneAsset.h
new file mode 100644
--- /dev/null
+++ b/Engine/Plugins/Experimental/Animation/ContextualAnimation/Source/ContextualAnimation/Public/ContextualAnimCompositeSceneAsset.h
@@ -0,0 +1,91 @@
+// Copyright Epic Games, Inc. All Rights Reserved.
+
+#pragma once
+
+#include "ContextualAnimSceneAssetBase.h"
+#include "ContextualAnimCompositeSceneAsset.generated.h"
+
+/** Stores the result of a query function */
+USTRUCT(BlueprintType)
+struct FContextualAnimQueryResult
+{
+	GENERATED_BODY()
+
+	UPROPERTY(BlueprintReadWrite, Category = "Defaults")
+	TWeakObjectPtr<UAnimMontage> Animation;
+
+	UPROPERTY(BlueprintReadWrite, Category = "Defaults")
+	FTransform EntryTransform;
+
+	UPROPERTY(BlueprintReadWrite, Category = "Defaults")
+	FTransform SyncTransform;
+
+	UPROPERTY(BlueprintReadWrite, Category = "Defaults")
+	float AnimStartTime = 0.f;
+
+	UPROPERTY(BlueprintReadWrite, Category = "Defaults")
+	int32 DataIndex = INDEX_NONE;
+
+	void Reset()
+	{
+		Animation.Reset();
+		EntryTransform = SyncTransform = FTransform::Identity;
+		AnimStartTime = 0.f;
+		DataIndex = INDEX_NONE;
+	}
+
+	FORCEINLINE bool IsValid() const { return DataIndex != INDEX_NONE; }
+};
+
+/** Stores the parameters passed into query function */
+USTRUCT(BlueprintType)
+struct FContextualAnimQueryParams
+{
+	GENERATED_BODY()
+
+	UPROPERTY(EditAnywhere, BlueprintReadWrite, Category = "Defaults")
+	TWeakObjectPtr<const AActor> Querier;
+
+	UPROPERTY(EditAnywhere, BlueprintReadWrite, Category = "Defaults")
+	FTransform QueryTransform;
+
+	UPROPERTY(EditAnywhere, BlueprintReadWrite, Category = "Defaults")
+	bool bComplexQuery = false;
+
+	UPROPERTY(EditAnywhere, BlueprintReadWrite, Category = "Defaults")
+	bool bFindAnimStartTime = false;
+
+	FContextualAnimQueryParams() {}
+
+	FContextualAnimQueryParams(const AActor* InQuerier, bool bInComplexQuery, bool bInFindAnimStartTime)
+		: Querier(InQuerier), bComplexQuery(bInComplexQuery), bFindAnimStartTime(bInFindAnimStartTime) {}
+
+	FContextualAnimQueryParams(const FTransform& InQueryTransform, bool bInComplexQuery, bool bInFindAnimStartTime)
+		: QueryTransform(InQueryTransform), bComplexQuery(bInComplexQuery), bFindAnimStartTime(bInFindAnimStartTime) {}
+};
+
+UCLASS()
+class CONTEXTUALANIMATION_API UContextualAnimCompositeSceneAsset : public UContextualAnimSceneAssetBase
+{
+	GENERATED_BODY()
+
+public:
+
+	static const FName InteractorRoleName;
+	static const FName InteractableRoleName;
+
+	UPROPERTY(EditAnywhere, BlueprintReadWrite, Category = "Animations")
+	FContextualAnimCompositeTrack InteractorTrack;
+
+	UPROPERTY(EditAnywhere, BlueprintReadWrite, Category = "Animations")
+	FContextualAnimTrack InteractableTrack;
+
+	UContextualAnimCompositeSceneAsset(const FObjectInitializer& ObjectInitializer);
+
+	virtual void PreSave(const class ITargetPlatform* TargetPlatform) override;
+
+	bool QueryData(FContextualAnimQueryResult& OutResult, const FContextualAnimQueryParams& QueryParams, const FTransform& ToWorldTransform) const;
+
+	virtual UClass* GetPreviewActorClassForRole(const FName& Role) const override;
+	virtual EContextualAnimJoinRule GetJoinRuleForRole(const FName& Role) const override;
+};
diff --git a/Engine/Plugins/Experimental/Animation/ContextualAnimation/Source/ContextualAnimation/Public/ContextualAnimManager.h b/Engine/Plugins/Experimental/Animation/ContextualAnimation/Source/ContextualAnimation/Public/ContextualAnimManager.h
new file mode 100644
--- /dev/null
+++ b/Engine/Plugins/Experimental/Animation/ContextualAnimation/Source/ContextualAnimation/Public/ContextualAnimManager.h
@@ -0,0 +1,42 @@
+// Copyright Epic Games, Inc. All Rights Reserved.
+
+#pragma once
+
+#include "CoreMinimal.h"
+#include "UObject/Object.h"
+#include "ContextualAnimManager.generated.h"
+
+class UContextualAnimSceneAssetBase;
+class UContextualAnimSceneInstance;
+class UWorld;
+
+UCLASS()
+class CONTEXTUALANIMATION_API UContextualAnimManager : public UObject
+{
+	GENERATED_BODY()
+
+public:
+
+	UContextualAnimManager(const FObjectInitializer& ObjectInitializer);
+
+	virtual UWorld* GetWorld() const override;
+
+	void Tick(float DeltaTime);
+
+	bool TryStartScene(const UContextualAnimSceneAssetBase* SceneAsset, const TMap<FName, AActor*>& Bindings);
+
+	bool TryStopSceneWithActor(AActor* Actor);
+
+	bool IsActorInAnyScene(AActor* Actor) const;
+
+	UContextualAnimSceneInstance* GetSceneWithActor(AActor* Actor);
+
+	static UContextualAnimManager* Get(const UWorld* World);
+
+protected:
+
+	UPROPERTY()
+	TArray<UContextualAnimSceneInstance*> Instances;
+
+	void OnSceneInstanceEnded(UContextualAnimSceneInstance* SceneInstance);
+};
\ No newline at end of file
diff --git a/Engine/Plugins/Experimental/Animation/ContextualAnimation/Source/ContextualAnimation/Public/ContextualAnimMetadata.h b/Engine/Plugins/Experimental/Animation/ContextualAnimation/Source/ContextualAnimation/Public/ContextualAnimMetadata.h
new file mode 100644
--- /dev/null
+++ b/Engine/Plugins/Experimental/Animation/ContextualAnimation/Source/ContextualAnimation/Public/ContextualAnimMetadata.h
@@ -0,0 +1,61 @@
+// Copyright Epic Games, Inc. All Rights Reserved.
+
+#pragma once
+
+#include "CoreMinimal.h"
+#include "UObject/Object.h"
+#include "ContextualAnimTypes.h"
+#include "ContextualAnimMetadata.generated.h"
+
+struct FContextualAnimData;
+
+USTRUCT(BlueprintType)
+struct FContextualAnimDistanceParam
+{
+	GENERATED_BODY()
+
+	UPROPERTY(VisibleAnywhere, BlueprintReadWrite, Category = "Defaults", meta = (UIMin = 0, ClampMin = 0))
+	float Value = 0.f;
+
+	UPROPERTY(EditAnywhere, BlueprintReadWrite, Category = "Defaults", meta = (UIMin = 0, ClampMin = 0))
+	float MinDistance = 0.f;
+
+	UPROPERTY(EditAnywhere, BlueprintReadWrite, Category = "Defaults", meta = (UIMin = 0, ClampMin = 0))
+	float MaxDistance = 0.f;
+};
+
+USTRUCT(BlueprintType)
+struct FContextualAnimAngleParam
+{
+	GENERATED_BODY()
+
+	UPROPERTY(VisibleAnywhere, BlueprintReadWrite, Category = "Defaults")
+	float Value = 0.f;
+
+	UPROPERTY(EditAnywhere, BlueprintReadWrite, Category = "Defaults")
+	float Tolerance = 0.f;
+};
+
+UCLASS(BlueprintType, EditInlineNew)
+class CONTEXTUALANIMATION_API UContextualAnimMetadata : public UObject
+{
+	GENERATED_BODY()
+
+public:
+
+	UPROPERTY(EditAnywhere, BlueprintReadWrite, Category = "Defaults")
+	float OffsetFromOrigin = 0.f;
+
+	UPROPERTY(EditAnywhere, BlueprintReadWrite, Category = "Defaults")
+	FContextualAnimDistanceParam Distance;
+
+	UPROPERTY(EditAnywhere, BlueprintReadWrite, Category = "Defaults")
+	FContextualAnimAngleParam Angle;
+
+	UPROPERTY(EditAnywhere, BlueprintReadWrite, Category = "Defaults")
+	FContextualAnimAngleParam Facing;
+
+	UContextualAnimMetadata(const FObjectInitializer& ObjectInitializer);
+
+	class UContextualAnimSceneAssetBase* GetSceneAssetOwner() const;
+};
\ No newline at end of file
diff --git a/Engine/Plugins/Experimental/Animation/ContextualAnimation/Source/ContextualAnimation/Public/ContextualAnimSceneActorComponent.h b/Engine/Plugins/Experimental/Animation/ContextualAnimation/Source/ContextualAnimation/Public/ContextualAnimSceneActorComponent.h
new file mode 100644
--- /dev/null
+++ b/Engine/Plugins/Experimental/Animation/ContextualAnimation/Source/ContextualAnimation/Public/ContextualAnimSceneActorComponent.h
@@ -0,0 +1,49 @@
+// Copyright Epic Games, Inc. All Rights Reserved.
+
+#pragma once
+
+#include "CoreMinimal.h"
+#include "Components/SphereComponent.h"
+#include "ContextualAnimCompositeSceneAsset.h"
+#include "ContextualAnimSceneActorComponent.generated.h"
+
+class UAnimInstance;
+class UAnimMontage;
+class AActor;
+
+USTRUCT(BlueprintType)
+struct FContextualAnimDebugParams
+{
+	GENERATED_BODY()
+
+	UPROPERTY(EditAnywhere, BlueprintReadWrite, Category = "Debug")
+	TWeakObjectPtr<AActor> TestActor;
+
+	UPROPERTY(EditAnywhere, BlueprintReadWrite, Category = "Debug", meta = (UIMin = 0, ClampMin = 0))
+	float DrawAlignmentTransformAtTime = 0.f;
+};
+
+UCLASS(meta = (BlueprintSpawnableComponent))
+class CONTEXTUALANIMATION_API UContextualAnimSceneActorComponent : public USphereComponent
+{
+	GENERATED_BODY()
+
+public:
+
+	UContextualAnimSceneActorComponent(const FObjectInitializer& ObjectInitializer);
+
+	UPROPERTY(EditAnywhere, BlueprintReadWrite, Category = "Config")
+	UContextualAnimCompositeSceneAsset* SceneAsset;
+
+	UPROPERTY(EditAnywhere, BlueprintReadWrite, Category = "Debug")
+	bool bEnableDebug;
+
+	UPROPERTY(EditAnywhere, BlueprintReadWrite, Category = "Debug", meta = (ShowOnlyInnerProperties, EditCondition = "bEnableDebug"))
+	FContextualAnimDebugParams DebugParams;
+
+	virtual FPrimitiveSceneProxy* CreateSceneProxy() override;
+	virtual FBoxSphereBounds CalcBounds(const FTransform& LocalToWorld) const;
+
+	UFUNCTION(BlueprintCallable, Category = "Contextual Animation System")
+	bool QueryData(const FContextualAnimQueryParams& QueryParams, FContextualAnimQueryResult& Result) const;
+};
diff --git a/Engine/Plugins/Experimental/Animation/ContextualAnimation/Source/ContextualAnimation/Public/ContextualAnimSceneAsset.h b/Engine/Plugins/Experimental/Animation/ContextualAnimation/Source/ContextualAnimation/Public/ContextualAnimSceneAsset.h
new file mode 100644
--- /dev/null
+++ b/Engine/Plugins/Experimental/Animation/ContextualAnimation/Source/ContextualAnimation/Public/ContextualAnimSceneAsset.h
@@ -0,0 +1,27 @@
+// Copyright Epic Games, Inc. All Rights Reserved.
+
+#pragma once
+
+#include "ContextualAnimSceneAssetBase.h"
+#include "ContextualAnimSceneAsset.generated.h"
+
+UCLASS()
+class CONTEXTUALANIMATION_API UContextualAnimSceneAsset : public UContextualAnimSceneAssetBase
+{
+	GENERATED_BODY()
+
+public:
+
+	UPROPERTY(EditAnywhere, BlueprintReadWrite, Category = "Animations")
+	TMap<FName, FContextualAnimTrack> DataContainer;
+
+	UContextualAnimSceneAsset(const FObjectInitializer& ObjectInitializer);
+
+	virtual void PreSave(const class ITargetPlatform* TargetPlatform) override;
+
+	virtual UClass* GetPreviewActorClassForRole(const FName& Role) const override;
+	
+	virtual EContextualAnimJoinRule GetJoinRuleForRole(const FName& Role) const override;
+	
+	const FContextualAnimTrack* FindTrack(const FName& Role) const;
+};
diff --git a/Engine/Plugins/Experimental/Animation/ContextualAnimation/Source/ContextualAnimation/Public/ContextualAnimSceneAssetBase.h b/Engine/Plugins/Experimental/Animation/ContextualAnimation/Source/ContextualAnimation/Public/ContextualAnimSceneAssetBase.h
new file mode 100644
--- /dev/null
+++ b/Engine/Plugins/Experimental/Animation/ContextualAnimation/Source/ContextualAnimation/Public/ContextualAnimSceneAssetBase.h
@@ -0,0 +1,63 @@
+// Copyright Epic Games, Inc. All Rights Reserved.
+
+#pragma once
+
+#include "Engine/DataAsset.h"
+#include "ContextualAnimTypes.h"
+#include "ContextualAnimSceneAssetBase.generated.h"
+
+class UContextualAnimScenePivotProvider;
+
+USTRUCT(BlueprintType)
+struct FContextualAnimAlignmentSectionData
+{
+	GENERATED_BODY()
+
+	UPROPERTY(EditAnywhere, BlueprintReadWrite, Category = "Alignment")
+	FName SectionName = NAME_None;
+
+	UPROPERTY(EditAnywhere, Instanced, Category = "Alignment")
+	UContextualAnimScenePivotProvider* ScenePivotProvider = nullptr;
+
+	UPROPERTY(VisibleAnywhere, BlueprintReadWrite, Category = "Alignment", AdvancedDisplay)
+	FTransform ScenePivot;
+};
+
+UCLASS(Abstract, Blueprintable)
+class CONTEXTUALANIMATION_API UContextualAnimSceneAssetBase : public UDataAsset
+{
+	GENERATED_BODY()
+
+public:
+
+	UPROPERTY(EditAnywhere, BlueprintReadWrite, Category = "Alignment")
+	TArray<FContextualAnimAlignmentSectionData> AlignmentSections;
+
+	UPROPERTY(EditAnywhere, BlueprintReadWrite, Category = "Alignment", meta = (ClampMin = "1", ClampMax = "60"), AdvancedDisplay)
+	int32 SampleRate;
+
+	UPROPERTY(EditAnywhere, BlueprintReadWrite, Category = "Alignment", AdvancedDisplay)
+	FTransform MeshToComponent;
+
+	UPROPERTY(EditAnywhere, BlueprintReadWrite, Category = "Roles")
+	FName PrimaryRole;
+
+	UContextualAnimSceneAssetBase(const FObjectInitializer& ObjectInitializer);
+
+	virtual void PreSave(const class ITargetPlatform* TargetPlatform) override;
+
+	//@TODO: Make pure virtual
+	virtual UClass* GetPreviewActorClassForRole(const FName& Role) const { return nullptr; }
+	virtual EContextualAnimJoinRule GetJoinRuleForRole(const FName& Role) const { return EContextualAnimJoinRule::Default; }
+
+	FTransform ExtractTransformFromAnimData(const FContextualAnimData& AnimData, float Time) const;
+
+	FORCEINLINE FName GetAlignmentSectionNameAtIndex(int32 Index) const 
+	{
+		return AlignmentSections.IsValidIndex(Index) ? AlignmentSections[Index].SectionName : NAME_None;
+	}
+
+protected:
+
+	void GenerateAlignmentTracksRelativeToScenePivot(FContextualAnimData& AnimData);
+};
diff --git a/Engine/Plugins/Experimental/Animation/ContextualAnimation/Source/ContextualAnimation/Public/ContextualAnimSceneInstance.h b/Engine/Plugins/Experimental/Animation/ContextualAnimation/Source/ContextualAnimation/Public/ContextualAnimSceneInstance.h
new file mode 100644
--- /dev/null
+++ b/Engine/Plugins/Experimental/Animation/ContextualAnimation/Source/ContextualAnimation/Public/ContextualAnimSceneInstance.h
@@ -0,0 +1,116 @@
+// Copyright Epic Games, Inc. All Rights Reserved.
+
+#pragma once
+
+#include "CoreMinimal.h"
+#include "UObject/Object.h"
+#include "ContextualAnimSceneInstance.generated.h"
+
+struct FContextualAnimData;
+class UContextualAnimSceneAssetBase;
+class UAnimMontage;
+class UWorld;
+
+/** Represent an actor bound to a role in the scene */
+USTRUCT()
+struct CONTEXTUALANIMATION_API FContextualAnimSceneActorData
+{
+	GENERATED_BODY()
+
+	/** The actual actor in the world */
+	TWeakObjectPtr<AActor> Actor = nullptr;
+
+	/** Ptr to the animation data in the scene asset used by this actor */
+	const FContextualAnimData* AnimDataPtr = nullptr;
+
+	/** Desired time to start the animation */
+	float AnimStartTime = 0.f;
+
+	FContextualAnimSceneActorData() : Actor(nullptr), AnimDataPtr(nullptr), AnimStartTime(0.f) {}
+
+	FContextualAnimSceneActorData(AActor* InActor, const FContextualAnimData* InAnimData, float InAnimStartTime = 0.f)
+		: Actor(InActor), AnimDataPtr(InAnimData), AnimStartTime(InAnimStartTime)
+	{}
+
+	/** Return a pointer to the actual actor in the world */
+	FORCEINLINE AActor* GetActor() const { return Actor.Get(); }
+
+	/** Return the transform used for alignment for this scene actor */
+	FTransform GetTransform() const;
+
+	/** Return the current playback time of the animation this actor is playing */
+	float GetAnimTime() const;
+};
+
+/** Delegate to notify external objects when this is scene is completed */
+DECLARE_DELEGATE_OneParam(FOnContextualAnimSceneEnded, class UContextualAnimSceneInstance*)
+
+/** Delegate to notify external objects when an actor join this scene */
+DECLARE_DELEGATE_TwoParams(FOnContextualAnimSceneActorJoined, class UContextualAnimSceneInstance*, AActor* Actor)
+
+/** Delegate to notify external objects when an actor left this scene */
+DECLARE_DELEGATE_TwoParams(FOnContextualAnimSceneActorLeft, class UContextualAnimSceneInstance*, AActor* Actor)
+
+/** Instance of a contextual animation scene */
+UCLASS()
+class CONTEXTUALANIMATION_API UContextualAnimSceneInstance : public UObject
+{
+	GENERATED_BODY()
+
+public:
+
+	/** Scene asset this instance was created from */
+	UPROPERTY()
+	TObjectPtr<const UContextualAnimSceneAssetBase> SceneAsset;
+
+	/** Map of roles to scene actor */
+	UPROPERTY()
+	TMap<FName, FContextualAnimSceneActorData> SceneActorMap;
+
+	/** List of alignment section to scene pivot */
+	TArray<TTuple<FName, FTransform>> AlignmentSectionToScenePivotList;
+
+	/** Delegate to notify external objects when this is scene is completed */
+	FOnContextualAnimSceneEnded OnSceneEnded;
+
+	/** Delegate to notify external objects when an actor join */
+	FOnContextualAnimSceneActorJoined OnActorJoined;
+
+	/** Delegate to notify external objects when an actor leave */
+	FOnContextualAnimSceneActorLeft OnActorLeft;
+
+	UContextualAnimSceneInstance(const FObjectInitializer& ObjectInitializer);
+
+	virtual UWorld* GetWorld() const override;
+
+	void Tick(float DeltaTime);
+
+	/** Resolve initial alignment and start playing animation for all actors */
+	void Start();
+
+	/** Force all the actors to leave the scene */
+	void Stop();
+
+	/** Whether the supplied actor is part of this scene */
+	bool IsActorInThisScene(const AActor* Actor) const;
+
+protected:
+
+	/** Tells the scene actor to join the scene (play animation) */
+	void Join(FContextualAnimSceneActorData& Data);
+
+	/** Tells the scene actor to leave the scene (stop animation) */
+	void Leave(FContextualAnimSceneActorData& Data);
+
+	/** Helper function to set ignore collision between the supplied actor and all the other actors in this scene */
+	void SetIgnoreCollisionWithOtherActors(AActor* Actor, bool bValue) const;
+
+	UFUNCTION()
+	void OnMontageBlendingOut(UAnimMontage* Montage, bool bInterrupted);
+
+	UFUNCTION()
+	void OnNotifyBeginReceived(FName NotifyName, const FBranchingPointNotifyPayload& BranchingPointNotifyPayload);
+
+	UFUNCTION()
+	void OnNotifyEndReceived(FName NotifyName, const FBranchingPointNotifyPayload& BranchingPointNotifyPayload);
+};
\ No newline at end of file
diff --git a/Engine/Plugins/Experimental/Animation/ContextualAnimation/Source/ContextualAnimation/Public/ContextualAnimScenePivotProvider.h b/Engine/Plugins/Experimental/Animation/ContextualAnimation/Source/ContextualAnimation/Public/ContextualAnimScenePivotProvider.h
new file mode 100644
--- /dev/null
+++ b/Engine/Plugins/Experimental/Animation/ContextualAnimation/Source/ContextualAnimation/Public/ContextualAnimScenePivotProvider.h
@@ -0,0 +1,63 @@
+// Copyright Epic Games, Inc. All Rights Reserved.
+
+#pragma once
+
+#include "CoreMinimal.h"
+#include "UObject/Object.h"
+#include "ContextualAnimSceneInstance.h"
+#include "ContextualAnimScenePivotProvider.generated.h"
+
+class UContextualAnimSceneAssetBase;
+
+UCLASS(Abstract, BlueprintType, EditInlineNew)
+class CONTEXTUALANIMATION_API UContextualAnimScenePivotProvider : public UObject
+{
+	GENERATED_BODY()
+
+public:
+
+	UContextualAnimScenePivotProvider(const FObjectInitializer& ObjectInitializer);
+
+	virtual FTransform CalculateScenePivot_Source() const { return FTransform::Identity; }
+	virtual FTransform CalculateScenePivot_Runtime(const TMap<FName, FContextualAnimSceneActorData>& SceneActorMap) const { return FTransform::Identity; }
+
+	const UContextualAnimSceneAssetBase* GetSceneAsset() const;
+};
+
+UCLASS()
+class CONTEXTUALANIMATION_API UContextualAnimScenePivotProvider_Default : public UContextualAnimScenePivotProvider
+{
+	GENERATED_BODY()
+
+public:
+
+	UPROPERTY(EditAnywhere, BlueprintReadWrite, Category = "Settings")
+	FName PrimaryRole;
+
+	UPROPERTY(EditAnywhere, BlueprintReadWrite, Category = "Settings")
+	FName SecondaryRole;
+
+	UPROPERTY(EditAnywhere, BlueprintReadWrite, Category = "Settings")
+	float Weight = 0.5f;
+
+	UContextualAnimScenePivotProvider_Default(const FObjectInitializer& ObjectInitializer);
+
+	virtual FTransform CalculateScenePivot_Source() const override;
+	virtual FTransform CalculateScenePivot_Runtime(const TMap<FName, FContextualAnimSceneActorData>& SceneActorMap) const override;
+};
+
+UCLASS()
+class CONTEXTUALANIMATION_API UContextualAnimScenePivotProvider_RelativeTo : public UContextualAnimScenePivotProvider
+{
+	GENERATED_BODY()
+
+public:
+
+	UPROPERTY(EditAnywhere, BlueprintReadWrite, Category = "Settings")
+	FName RelativeToRole;
+
+	UContextualAnimScenePivotProvider_RelativeTo(const FObjectInitializer& ObjectInitializer);
+
+	virtual FTransform CalculateScenePivot_Source() const override;
+	virtual FTransform CalculateScenePivot_Runtime(const TMap<FName, FContextualAnimSceneActorData>& SceneActorMap) const override;
+};
\ No newline at end of file
diff --git a/Engine/Plugins/Experimental/Animation/ContextualAnimation/Source/ContextualAnimation/Public/ContextualAnimTypes.h b/Engine/Plugins/Experimental/Animation/ContextualAnimation/Source/ContextualAnimation/Public/ContextualAnimTypes.h
new file mode 100644
--- /dev/null
+++ b/Engine/Plugins/Experimental/Animation/ContextualAnimation/Source/ContextualAnimation/Public/ContextualAnimTypes.h
@@ -0,0 +1,110 @@
+// Copyright Epic Games, Inc. All Rights Reserved.
+
+#pragma once
+
+#include "Animation/AnimSequence.h"
+#include "Templates/SubclassOf.h"
+#include "ContextualAnimTypes.generated.h"
+
+DECLARE_LOG_CATEGORY_EXTERN(LogContextualAnim, Log, All);
+
+class AActor;
+class UAnimMontage;
+class UContextualAnimMetadata;
+
+/** Container for alignment tracks */
+USTRUCT()
+struct CONTEXTUALANIMATION_API FContextualAnimAlignmentTrackContainer
+{
+	GENERATED_BODY()
+
+	UPROPERTY()
+	FAnimSequenceTrackContainer Tracks;
+
+	UPROPERTY()
+	float SampleInterval = 0.f;
+
+	FTransform ExtractTransformAtTime(int32 TrackIndex, float Time) const;
+	FTransform ExtractTransformAtTime(const FName& TrackName, float Time) const;
+};
+
+
+USTRUCT(BlueprintType)
+struct CONTEXTUALANIMATION_API FContextualAnimData
+{
+	GENERATED_BODY()
+
+	UPROPERTY(EditAnywhere, BlueprintReadWrite, Category = "Defaults")
+	UAnimMontage* Animation = nullptr;
+
+	UPROPERTY(EditAnywhere, BlueprintReadWrite, Category = "Defaults")
+	float AnimMaxStartTime = 0.f;
+
+	UPROPERTY()
+	FContextualAnimAlignmentTrackContainer AlignmentData;
+
+	UPROPERTY(EditAnywhere, Instanced, Category = "Defaults")
+	UContextualAnimMetadata* Metadata = nullptr;
+
+	UPROPERTY(EditAnywhere, BlueprintReadWrite, Category = "Defaults")
+	FTransform MeshToScene;
+
+	/** deprecated */
+	UPROPERTY()
+	float SyncTime = 0.f;
+
+	float GetSyncTimeForWarpSection(int32 WarpSectionIndex) const;
+	float GetSyncTimeForWarpSection(const FName& WarpSectionName) const;
+
+	FORCEINLINE FTransform GetAlignmentTransformAtTime(float Time) const { return AlignmentData.ExtractTransformAtTime(0, Time); }
+	FORCEINLINE FTransform GetAlignmentTransformAtEntryTime() const { return AlignmentData.ExtractTransformAtTime(0, 0.f); }
+	FORCEINLINE FTransform GetAlignmentTransformAtSyncTime() const { return AlignmentData.ExtractTransformAtTime(0, GetSyncTimeForWarpSection(0)); }
+
+	float FindBestAnimStartTime(const FVector& LocalLocation) const;
+
+	static const FContextualAnimData EmptyAnimData;
+};
+
+/** Defines when the actor should start playing the animation */
+UENUM(BlueprintType)
+enum class EContextualAnimJoinRule : uint8
+{
+	Default,
+	Late
+};
+
+USTRUCT(BlueprintType)
+struct CONTEXTUALANIMATION_API FContextualAnimTrackSettings
+{
+	GENERATED_BODY()
+
+	UPROPERTY(EditAnywhere, BlueprintReadWrite, Category = "Defaults")
+	TSubclassOf<AActor> PreviewActorClass;
+
+	UPROPERTY(EditAnywhere, BlueprintReadWrite, Category = "Defaults")
+	EContextualAnimJoinRule JoinRule = EContextualAnimJoinRule::Default;
+};
+
+USTRUCT(BlueprintType)
+struct FContextualAnimTrack
+{
+	GENERATED_BODY()
+
+	UPROPERTY(EditAnywhere, BlueprintReadWrite, Category = "Defaults")
+	FContextualAnimTrackSettings Settings;
+
+	UPROPERTY(EditAnywhere, BlueprintReadWrite, Category = "Defaults")
+	FContextualAnimData AnimData;
+};
+
+USTRUCT(BlueprintType)
+struct FContextualAnimCompositeTrack
+{
+	GENERATED_BODY()
+
+	UPROPERTY(EditAnywhere, BlueprintReadWrite, Category = "Defaults")
+	FContextualAnimTrackSettings Settings;
+
+	UPROPERTY(EditAnywhere, BlueprintReadWrite, Category = "Defaults")
+	TArray<FContextualAnimData> AnimDataContainer;
+};
\ No newline at end of file
diff --git a/Engine/Plugins/Experimental/Animation/ContextualAnimation/Source/ContextualAnimation/Public/ContextualAnimUtilities.h b/Engine/Plugins/Experimental/Animation/ContextualAnimation/Source/ContextualAnimation/Public/ContextualAnimUtilities.h
new file mode 100644
--- /dev/null
+++ b/Engine/Plugins/Experimental/Animation/ContextualAnimation/Source/ContextualAnimation/Public/ContextualAnimUtilities.h
@@ -0,0 +1,24 @@
+// Copyright Epic Games, Inc. All Rights Reserved.
+
+#pragma once
+
+#include "Kismet/BlueprintFunctionLibrary.h"
+#include "Animation/AnimSequence.h"
+#include "ContextualAnimUtilities.generated.h"
+
+class UContextualAnimSceneAsset;
+
+UCLASS()
+class CONTEXTUALANIMATION_API UContextualAnimUtilities : public UBlueprintFunctionLibrary
+{
+	GENERATED_BODY()
+
+public:
+
+	static void ExtractLocalSpacePose(const UAnimSequenceBase* Animation, const FBoneContainer& BoneContainer, float Time, bool bExtractRootMotion, FCompactPose& OutPose);
+	static void ExtractComponentSpacePose(const UAnimSequenceBase* Animation, const FBoneContainer& BoneContainer, float Time, bool bExtractRootMotion, FCSPose<FCompactPose>& OutPose);
+	static FTransform ExtractRootMotionFromAnimation(const UAnimSequenceBase* Animation, float StartTime, float EndTime);
+
+	static void DrawDebugPose(const UWorld* World, const UAnimSequenceBase* Animation, float Time, const FTransform& LocalToWorldTransform, const FColor& Color, float LifeTime, float Thickness);
+	static void DrawDebugScene(const UWorld* World, const UContextualAnimSceneAsset* SceneAsset, float Time, const FTransform& ToWorldTransform, const FColor& Color, float LifeTime, float Thickness);
+};
\ No newline at end of file
diff --git a/Engine/Plugins/Experimental/Animation/ContextualAnimation/Source/ContextualAnimation/Public/ContextualAnimation.h b/Engine/Plugins/Experimental/Animation/ContextualAnimation/Source/ContextualAnimation/Public/ContextualAnimation.h
--- a/Engine/Plugins/Experimental/Animation/ContextualAnimation/Source/ContextualAnimation/Public/ContextualAnimation.h
+++ b/Engine/Plugins/Experimental/Animation/ContextualAnimation/Source/ContextualAnimation/Public/ContextualAnimation.h
@@ -3,13 +3,37 @@
 #pragma once
 
 #include "CoreMinimal.h"
+#include "UObject/GCObject.h"
 #include "Modules/ModuleManager.h"
+#include "Engine/World.h"
 
-class FContextualAnimationModule : public IModuleInterface
+class UContextualAnimManager;
+
+class FContextualAnimationModule : public IModuleInterface, public FGCObject
 {
 public:
 
 	/** IModuleInterface implementation */
 	virtual void StartupModule() override;
 	virtual void ShutdownModule() override;
+
+	// Begin FGCObject overrides
+	virtual void AddReferencedObjects(FReferenceCollector& Collector) override;
+	// End FGCObject overrides
+
+	FORCEINLINE static UContextualAnimManager* GetManager(const UWorld* World)
+	{
+		return WorldToManagerMap.FindRef(World);
+	}
+
+private:
+
+	static void OnWorldInit(UWorld* World, const UWorld::InitializationValues IVS);
+
+	static void OnWorldCleanup(UWorld* World, bool bSessionEnded, bool bCleanupResources);
+
+	static TMap<const UWorld*, UContextualAnimManager*> WorldToManagerMap;
+
+	FDelegateHandle OnPreWorldInitDelegateHandle;
+	FDelegateHandle OnPostWorldCleanupDelegateHandle;
 };
diff --git a/Engine/Plugins/Experimental/Animation/ContextualAnimation/Source/ContextualAnimationEditor/Private/ContextualAnimEdMode.cpp b/Engine/Plugins/Experimental/Animation/ContextualAnimation/Source/ContextualAnimationEditor/Private/ContextualAnimEdMode.cpp
--- a/Engine/Plugins/Experimental/Animation/ContextualAnimation/Source/ContextualAnimationEditor/Private/ContextualAnimEdMode.cpp
+++ b/Engine/Plugins/Experimental/Animation/ContextualAnimation/Source/ContextualAnimationEditor/Private/ContextualAnimEdMode.cpp
@@ -9,22 +9,33 @@
 #include "EngineUtils.h"
 #include "ContextualAnimEdModeSettings.h"
 #include "DrawDebugHelpers.h"
+#include "ContextualAnimSceneAsset.h"
+#include "ContextualAnimCompositeSceneAsset.h"
+#include "ContextualAnimUtilities.h"
+#include "ContextualAnimManager.h"
+#include "ContextualAnimPreviewManager.h"
+#include "ContextualAnimSceneActorComponent.h"
 #include "GameFramework/Character.h"
-#include "GameFramework/CharacterMovementComponent.h"
-#include "NavigationSystem.h"
-#include "AIController.h"
-#include "ContextualAnimComponent.h"
+#include "ContextualAnimMetadata.h"
 
 const FEditorModeID FContextualAnimEdMode::EM_ContextualAnimEdModeId = TEXT("EM_ContextualAnimEdMode");
 
 FContextualAnimEdMode::FContextualAnimEdMode()
 {
+	PreviewManager = NewObject<UContextualAnimPreviewManager>(UContextualAnimPreviewManager::StaticClass());
 }
 
 FContextualAnimEdMode::~FContextualAnimEdMode()
 {
 }
 
+void FContextualAnimEdMode::AddReferencedObjects(FReferenceCollector& Collector)
+{
+	FEdMode::AddReferencedObjects(Collector);
+
+	Collector.AddReferencedObject(PreviewManager);
+}
+
 void FContextualAnimEdMode::Enter()
 {
 	FEdMode::Enter();
@@ -47,7 +58,7 @@ void FContextualAnimEdMode::Exit()
 	FEdMode::Exit();
 }
 
-TSharedPtr<FContextualAnimEdModeToolkit> FContextualAnimEdMode::GetContextualAnimEdModeToolkit()
+TSharedPtr<FContextualAnimEdModeToolkit> FContextualAnimEdMode::GetContextualAnimEdModeToolkit() const
 {
 	return StaticCastSharedPtr<FContextualAnimEdModeToolkit>(Toolkit);
 }
@@ -56,72 +67,90 @@ void FContextualAnimEdMode::Tick(FEditorViewportClient* ViewportClient, float De
 {
 	FEdMode::Tick(ViewportClient, DeltaTime);
 
-	if(TestCharacter.IsValid() && GEditor->IsSimulatingInEditor())
+	if(GEditor->IsSimulatingInEditor())
 	{
-		if (ViewportClient->Viewport->KeyState(EKeys::W) || ViewportClient->Viewport->KeyState(EKeys::S))
+		if (!ViewportClient->Viewport->KeyState(EKeys::RightMouseButton))
 		{
-			const FVector WorldDirection = FRotationMatrix(TestCharacter->GetActorRotation()).GetScaledAxis(EAxis::X);
-			TestCharacter->AddMovementInput(WorldDirection, ViewportClient->Viewport->KeyState(EKeys::W) ? 1.f : -1.f);
-		}
+			if (ViewportClient->Viewport->KeyState(EKeys::W) || ViewportClient->Viewport->KeyState(EKeys::S))
+			{
+				PreviewManager->MoveForward(ViewportClient->Viewport->KeyState(EKeys::W) ? 1.f : -1.f);
+			}
 
-		if (ViewportClient->Viewport->KeyState(EKeys::A) || ViewportClient->Viewport->KeyState(EKeys::D))
+			if (ViewportClient->Viewport->KeyState(EKeys::A) || ViewportClient->Viewport->KeyState(EKeys::D))
+			{
+				PreviewManager->MoveRight(ViewportClient->Viewport->KeyState(EKeys::D) ? 1.f : -1.f);
+			}
+		}
+	}
+	else
+	{
+		if(PreviewManager)
 		{
-			const FVector WorldDirection = FRotationMatrix(TestCharacter->GetActorRotation()).GetScaledAxis(EAxis::Y);
-			TestCharacter->AddMovementInput(WorldDirection, ViewportClient->Viewport->KeyState(EKeys::D) ? 1.f : -1.f);
+			if (PreviewManager->bDrawDebugScene)
+			{
+				if (UContextualAnimSceneAsset* Asset = GetContextualAnimEdModeToolkit()->GetSettings()->SceneAsset)
+				{
+					DrawDebugCoordinateSystem(GetWorld(), PreviewManager->ScenePivot.GetLocation(), PreviewManager->ScenePivot.Rotator(), 50.f, false, 0.f, 0, 1.f);
+					UContextualAnimUtilities::DrawDebugScene(GetWorld(), Asset, PreviewManager->Time, PreviewManager->ScenePivot, FColor::White, 0.f, 1.f);
+				}
+			}
+
+			//@TODO: Move this to an event that triggers after stop Simulating Mode (not sure if we have one)
+			if(PreviewManager->PreviewActors.Num())
+			{
+				PreviewManager->PreviewActors.Reset();
+			}
 		}
 	}
 }
 
 bool FContextualAnimEdMode::InputKey(FEditorViewportClient* ViewportClient, FViewport* Viewport, FKey Key, EInputEvent Event)
 {
-	if(Key == EKeys::Enter && Event == IE_Released && TestCharacter.IsValid() && GEditor->IsSimulatingInEditor())
+	if(Key == EKeys::Enter && Event == IE_Released && PreviewManager->TestCharacter.IsValid() && GEditor->IsSimulatingInEditor())
 	{
-		UContextualAnimComponent* ContextualAnimComp = nullptr;
-
-		TArray<UPrimitiveComponent*> OverlappingComps;
-		TestCharacter->GetOverlappingComponents(OverlappingComps);
+		UContextualAnimManager* Manager = UContextualAnimManager::Get(GetWorld());
+		check(Manager);		
 
-		for (UPrimitiveComponent* Comp : OverlappingComps)
+		if(!Manager->IsActorInAnyScene(PreviewManager->TestCharacter.Get()))
 		{
-			if (Comp && Comp->GetClass()->IsChildOf<UContextualAnimComponent>())
+			UContextualAnimSceneAsset* Asset = GetContextualAnimEdModeToolkit()->GetSettings()->SceneAsset;
+			if (Asset && PreviewManager->PreviewActors.Num() > 0)
 			{
-				ContextualAnimComp = Cast<UContextualAnimComponent>(Comp);
-				break;
+				Manager->TryStartScene(Asset, PreviewManager->PreviewActors);
 			}
-		}
-
-		if(ContextualAnimComp == nullptr)
-		{
-			GEngine->AddOnScreenDebugMessage(-1, 5.f, FColor::Yellow, FString(TEXT("WARNING: The preview actor is not overlapping any interactable")));
-			return FEdMode::InputKey(ViewportClient, Viewport, Key, Event);
-		}
-
-		// Check if the preview actor is already playing a contextual animation
-		if (!ContextualAnimComp->IsActorPlayingContextualAnimation(TestCharacter.Get()))
-		{
-			// Try to find best entry point
-			FContextualAnimQueryResult Result;
-			const bool bQueryResult = ContextualAnimComp->QueryData(FContextualAnimQueryParams(TestCharacter.Get(), true, true), Result);
-
-			// Early out of no valid entry point is found
-			if(bQueryResult == false)
+			else
 			{
-				GEngine->AddOnScreenDebugMessage(-1, 5.f, FColor::Yellow, FString::Printf(TEXT("WARNING: The preview actor is overlapping an interactable (%s) but there is no valid entry point for his current position"),
-					*GetNameSafe(ContextualAnimComp->GetOwner())));
-				return FEdMode::InputKey(ViewportClient, Viewport, Key, Event);
+				UContextualAnimSceneActorComponent* Comp = nullptr;
+
+				TArray<UPrimitiveComponent*> OverlappingComps;
+				PreviewManager->TestCharacter->GetOverlappingComponents(OverlappingComps);
+
+				for (UPrimitiveComponent* OverlappingComp : OverlappingComps)
+				{
+					if (OverlappingComp && OverlappingComp->GetClass()->IsChildOf<UContextualAnimSceneActorComponent>())
+					{
+						Comp = Cast<UContextualAnimSceneActorComponent>(OverlappingComp);
+						break;
+					}
+				}
+
+				if (Comp)
+				{
+					TMap<FName, AActor*> Bindings;
+					Bindings.Add(UContextualAnimCompositeSceneAsset::InteractableRoleName, Comp->GetOwner());
+					Bindings.Add(UContextualAnimCompositeSceneAsset::InteractorRoleName, PreviewManager->TestCharacter.Get());
+
+					UContextualAnimManager::Get(GetWorld())->TryStartScene(Comp->SceneAsset, Bindings);
+				}
+				else
+				{
+					GEngine->AddOnScreenDebugMessage(-1, 5.f, FColor::Yellow, FString(TEXT("WARNING: The preview actor is not overlapping any interactable")));
+				}
 			}
-
-			// QueryData does not load the animation and TryStartContextualAnimation expects the animation to be loaded
-			// @TODO: We may want to provide functions to query the data and load the animation asynchronous
-			Result.Animation.LoadSynchronous();
-
-			// Attempt to start the interaction
-			ContextualAnimComp->TryStartContextualAnimation(TestCharacter.Get(), Result);
 		}
 		else
 		{
-			//@TODO: For now we don't care where in the animation the preview actor is
-			ContextualAnimComp->TryEndContextualAnimation(TestCharacter.Get());
+			Manager->TryStopSceneWithActor(PreviewManager->TestCharacter.Get());
 		}
 
 		return true;
@@ -132,7 +161,7 @@ bool FContextualAnimEdMode::InputKey(FEditorViewportClient* ViewportClient, FVie
 
 bool FContextualAnimEdMode::HandleClick(FEditorViewportClient* InViewportClient, HHitProxy* HitProxy, const FViewportClick& Click)
 {
-	if (!Click.IsControlDown())
+	if (!Click.IsAltDown())
 	{
 		return FEdMode::HandleClick(InViewportClient, HitProxy, Click);
 	}
@@ -143,74 +172,59 @@ bool FContextualAnimEdMode::HandleClick(FEditorViewportClient* InViewportClient,
 		return FEdMode::HandleClick(InViewportClient, HitProxy, Click);
 	}
 
-	FSceneViewFamilyContext ViewFamily(FSceneViewFamily::ConstructionValues(InViewportClient->Viewport, InViewportClient->GetScene(), InViewportClient->EngineShowFlags).SetRealtimeUpdate(InViewportClient->IsRealtime()));
-	FSceneView* View = InViewportClient->CalcSceneView(&ViewFamily);
-	FViewportCursorLocation Cursor(View, InViewportClient, Click.GetClickPos().X, Click.GetClickPos().Y);
-	const auto ViewportType = InViewportClient->GetViewportType();
-
-	const FVector RayStart = Cursor.GetOrigin();
-	const FVector RayEnd = RayStart + Cursor.GetDirection() * HALF_WORLD_MAX;
-
 	FHitResult HitResult;
-	InViewportClient->GetWorld()->LineTraceSingleByChannel(HitResult, RayStart, RayEnd, ECC_WorldStatic, FCollisionQueryParams::DefaultQueryParam);
+	GetHitResultUnderCursor(HitResult, InViewportClient, Click);
 
 	GEngine->AddOnScreenDebugMessage(-1, 5.f, FColor::White, FString::Printf(TEXT("INFO: HandleClick: IsValidHit: %d Actor: %s Location: %s"),
 		HitResult.IsValidBlockingHit(), *GetNameSafe(HitResult.GetActor()), *HitResult.ImpactPoint.ToString()));
 
 	if (!HitResult.IsValidBlockingHit())
 	{
-		GEngine->AddOnScreenDebugMessage(-1, 5.f, FColor::Yellow, FString(TEXT("WARNING. Can't spawn preview actor. Reason: HitResult from click event is not a valid blocking hit.")));
+		GEngine->AddOnScreenDebugMessage(-1, 5.f, FColor::Yellow, FString(TEXT("WARNING. HitResult from click event is not a valid blocking hit.")));
 		return FEdMode::HandleClick(InViewportClient, HitProxy, Click);
-	}
-
-	if (!TestCharacter.IsValid())
+	}	
+	
+	FTransform SpawnTransform = FTransform(HitResult.ImpactPoint);
+	UContextualAnimSceneAsset* Asset = GetContextualAnimEdModeToolkit()->GetSettings()->SceneAsset;
+	if (Asset && PreviewManager->PreviewActors.Num() == 0)
 	{
-		TSubclassOf<ACharacter> TestActorClass = GetContextualAnimEdModeToolkit()->GetSettings()->TestCharacterClass;
-		if (!TestActorClass)
-		{
-			GEngine->AddOnScreenDebugMessage(-1, 5.f, FColor::Yellow, FString(TEXT("WARNING. Can't spawn preview actor. Reason: TestCharacterClass is invalid")));
-			return FEdMode::HandleClick(InViewportClient, HitProxy, Click);
-		}
-
-		TestCharacter = GetWorld()->SpawnActor<ACharacter>(TestActorClass, FTransform(HitResult.ImpactPoint));
-		if (ensureAlways(TestCharacter.IsValid()))
-		{
-			TestCharacter->bUseControllerRotationYaw = false;
-			if (UCharacterMovementComponent* CharacterMovementComp = TestCharacter->GetCharacterMovement())
-			{
-				CharacterMovementComp->bOrientRotationToMovement = true;
-				CharacterMovementComp->bUseControllerDesiredRotation = false;
-				CharacterMovementComp->RotationRate = FRotator(0.f, 540.0, 0.f);
-			}
-
-			if (!TestCharacter->AIControllerClass || !TestCharacter->AIControllerClass->IsChildOf<AAIController>())
-			{
-				TestCharacter->AIControllerClass = AAIController::StaticClass();
-			}
-
-			TestCharacter->SpawnDefaultController();
-		}
+		PreviewManager->SpawnPreviewActors(Asset, SpawnTransform); 
 	}
 	else
 	{
-		if (AAIController* Controller = Cast<AAIController>(TestCharacter->GetController()))
+		if (!PreviewManager->TestCharacter.IsValid())
 		{
-			UNavigationSystemV1* NavSys = FNavigationSystem::GetCurrent<UNavigationSystemV1>(Controller->GetWorld());
-			const ANavigationData* NavData = NavSys ? NavSys->GetNavDataForProps(Controller->GetNavAgentPropertiesRef(), Controller->GetNavAgentLocation()) : nullptr;
-
-			const bool bUsePathfinding = (NavData != nullptr);
-			Controller->MoveToLocation(HitResult.ImpactPoint, 10.f, true, bUsePathfinding);
+			PreviewManager->TestCharacter = Cast<ACharacter>(PreviewManager->SpawnPreviewActor(PreviewManager->DefaultPreviewClass, SpawnTransform));
+		}
+		else
+		{
+			PreviewManager->MoveToLocation(HitResult.ImpactPoint);
 		}
 	}
 
 	return true;
 }
 
+bool FContextualAnimEdMode::GetHitResultUnderCursor(FHitResult& OutHitResult, FEditorViewportClient* InViewportClient, const FViewportClick& Click) const
+{
+	FSceneViewFamilyContext ViewFamily(FSceneViewFamily::ConstructionValues(InViewportClient->Viewport, InViewportClient->GetScene(), InViewportClient->EngineShowFlags).SetRealtimeUpdate(InViewportClient->IsRealtime()));
+	FSceneView* View = InViewportClient->CalcSceneView(&ViewFamily);
+	FViewportCursorLocation Cursor(View, InViewportClient, Click.GetClickPos().X, Click.GetClickPos().Y);
+	const auto ViewportType = InViewportClient->GetViewportType();
+
+	const FVector RayStart = Cursor.GetOrigin();
+	const FVector RayEnd = RayStart + Cursor.GetDirection() * HALF_WORLD_MAX;
+
+	return InViewportClient->GetWorld()->LineTraceSingleByChannel(OutHitResult, RayStart, RayEnd, ECC_WorldStatic, FCollisionQueryParams::DefaultQueryParam);
+}
+
 bool FContextualAnimEdMode::UsesToolkits() const
 {
 	return true;
 }
 
-
-
+FContextualAnimEdMode& FContextualAnimEdMode::Get()
+{
+	return *(static_cast<FContextualAnimEdMode*>(GLevelEditorModeTools().GetActiveMode(FContextualAnimEdMode::EM_ContextualAnimEdModeId)));
+}
 
diff --git a/Engine/Plugins/Experimental/Animation/ContextualAnimation/Source/ContextualAnimationEditor/Private/ContextualAnimEdModeToolkit.cpp b/Engine/Plugins/Experimental/Animation/ContextualAnimation/Source/ContextualAnimationEditor/Private/ContextualAnimEdModeToolkit.cpp
--- a/Engine/Plugins/Experimental/Animation/ContextualAnimation/Source/ContextualAnimationEditor/Private/ContextualAnimEdModeToolkit.cpp
+++ b/Engine/Plugins/Experimental/Animation/ContextualAnimation/Source/ContextualAnimationEditor/Private/ContextualAnimEdModeToolkit.cpp
@@ -8,6 +8,7 @@
 #include "EditorModeManager.h"
 #include "Modules/ModuleManager.h"
 #include "ContextualAnimEdModeSettings.h"
+#include "ContextualAnimPreviewManager.h"
 
 #define LOCTEXT_NAMESPACE "ContextualAnimEdModeToolkit"
 
@@ -33,6 +34,9 @@ void FContextualAnimEdModeToolkit::Init(const TSharedPtr<IToolkitHost>& InitTool
 	EdModeSettingsWidget = PropertyModule.CreateDetailView(Args);
 	EdModeSettingsWidget->SetObject(Settings);
 
+	PreviewManagerWidget = PropertyModule.CreateDetailView(Args);
+	PreviewManagerWidget->SetObject(GetContextualAnimEdMode()->PreviewManager);
+
 	SAssignNew(ToolkitWidget, SBorder)
 		.HAlign(HAlign_Center)
 		.Padding(5)
@@ -49,8 +53,8 @@ void FContextualAnimEdModeToolkit::Init(const TSharedPtr<IToolkitHost>& InitTool
 				.Text(FText::FromString("\
 					- Select the class for your test actor\n\
 					- Start Simulating Mode\n\
-					- Ctrl + Click to spawn a test actor\n\
-					- Ctrl + Click or WASD to move the test actor around\n\
+					- Alt + Click to spawn a test actor\n\
+					- Alt + Click or WASD to move the test actor around\n\
 					- [Enter] to start / stop an interaction"))
 			]
 
@@ -60,6 +64,13 @@ void FContextualAnimEdModeToolkit::Init(const TSharedPtr<IToolkitHost>& InitTool
 			[
 				EdModeSettingsWidget.ToSharedRef()
 			]
+
+			+ SVerticalBox::Slot()
+			.HAlign(HAlign_Left)
+			.AutoHeight()
+			[
+				PreviewManagerWidget.ToSharedRef()
+			]
 		];
 
 	FModeToolkit::Init(InitToolkitHost);
diff --git a/Engine/Plugins/Experimental/Animation/ContextualAnimation/Source/ContextualAnimationEditor/Private/ContextualAnimPreviewManager.cpp b/Engine/Plugins/Experimental/Animation/ContextualAnimation/Source/ContextualAnimationEditor/Private/ContextualAnimPreviewManager.cpp
new file mode 100644
--- /dev/null
+++ b/Engine/Plugins/Experimental/Animation/ContextualAnimation/Source/ContextualAnimationEditor/Private/ContextualAnimPreviewManager.cpp
@@ -0,0 +1,107 @@
+// Copyright Epic Games, Inc. All Rights Reserved.
+
+#include "ContextualAnimPreviewManager.h"
+#include "ContextualAnimSceneAsset.h"
+#include "ContextualAnimEdMode.h"
+
+#include "GameFramework/Character.h"
+#include "GameFramework/CharacterMovementComponent.h"
+#include "NavigationSystem.h"
+#include "AIController.h"
+
+UContextualAnimPreviewManager::UContextualAnimPreviewManager(const FObjectInitializer& ObjectInitializer)
+	: Super(ObjectInitializer)
+{
+
+}
+
+void UContextualAnimPreviewManager::MoveForward(float Value)
+{
+	if (TestCharacter.IsValid())
+	{
+		const FVector WorldDirection = FRotationMatrix(TestCharacter->GetActorRotation()).GetScaledAxis(EAxis::X);
+		TestCharacter->AddMovementInput(WorldDirection, Value);
+	}
+}
+
+void UContextualAnimPreviewManager::MoveRight(float Value)
+{
+	if(TestCharacter.IsValid())
+	{
+		const FVector WorldDirection = FRotationMatrix(TestCharacter->GetActorRotation()).GetScaledAxis(EAxis::Y);
+		TestCharacter->AddMovementInput(WorldDirection, Value);
+	}
+}
+
+void UContextualAnimPreviewManager::MoveToLocation(const FVector& GoalLocation)
+{
+	AAIController* Controller = TestCharacter.IsValid() ? Cast<AAIController>(TestCharacter->GetController()) : nullptr;
+	if (Controller)
+	{
+		UNavigationSystemV1* NavSys = FNavigationSystem::GetCurrent<UNavigationSystemV1>(Controller->GetWorld());
+		const ANavigationData* NavData = NavSys ? NavSys->GetNavDataForProps(Controller->GetNavAgentPropertiesRef(), Controller->GetNavAgentLocation()) : nullptr;
+
+		const bool bUsePathfinding = (NavData != nullptr);
+		Controller->MoveToLocation(GoalLocation, 10.f, true, bUsePathfinding);
+	}
+}
+
+void UContextualAnimPreviewManager::SpawnPreviewActors(const UContextualAnimSceneAsset* SceneAsset, const FTransform& SceneOrigin)
+{
+	if (SceneAsset)
+	{
+		PreviewActors.Reset();
+
+		//DrawDebugCoordinateSystem(GetWorld(), ToWorldTransform.GetLocation(), ToWorldTransform.Rotator(), 20.f, false, 10.f, 0, 1.f);
+
+		for (const auto& Entry : SceneAsset->DataContainer)
+		{
+			const FName& Role = Entry.Key;
+			const FContextualAnimData& Data = Entry.Value.AnimData;
+
+			const FTransform SpawnTransform = (Data.AlignmentData.ExtractTransformAtTime(0, 0.f) * SceneOrigin);
+
+			UClass* PreviewClass = SceneAsset->GetPreviewActorClassForRole(Role);
+			if(!PreviewClass && DefaultPreviewClass)
+			{
+				PreviewClass = DefaultPreviewClass.Get();
+			}
+
+			AActor* PreviewActor = SpawnPreviewActor(PreviewClass, SpawnTransform);
+			if (PreviewActor)
+			{
+				PreviewActors.Add(Role, PreviewActor);
+
+				if (!TestCharacter.IsValid())
+				{
+					TestCharacter = Cast<ACharacter>(PreviewActor);
+				}
+			}
+		}
+	}
+}
+
+AActor* UContextualAnimPreviewManager::SpawnPreviewActor(UClass* Class, const FTransform& SpawnTransform) const
+{
+	AActor* PreviewActor = FContextualAnimEdMode::Get().GetWorld()->SpawnActor<AActor>(Class, SpawnTransform);
+
+	if (ACharacter* PreviewCharacter = Cast<ACharacter>(PreviewActor))
+	{
+		PreviewCharacter->bUseControllerRotationYaw = false;
+		if (UCharacterMovementComponent* CharacterMovementComp = PreviewCharacter->GetCharacterMovement())
+		{
+			CharacterMovementComp->bOrientRotationToMovement = true;
+			CharacterMovementComp->bUseControllerDesiredRotation = false;
+			CharacterMovementComp->RotationRate = FRotator(0.f, 540.0, 0.f);
+		}
+
+		if (!PreviewCharacter->AIControllerClass || !PreviewCharacter->AIControllerClass->IsChildOf<AAIController>())
+		{
+			PreviewCharacter->AIControllerClass = AAIController::StaticClass();
+		}
+
+		PreviewCharacter->SpawnDefaultController();
+	}
+
+	return PreviewActor;
+}
\ No newline at end of file
diff --git a/Engine/Plugins/Experimental/Animation/ContextualAnimation/Source/ContextualAnimationEditor/Private/ContextualAnimationEditor.cpp b/Engine/Plugins/Experimental/Animation/ContextualAnimation/Source/ContextualAnimationEditor/Private/ContextualAnimationEditor.cpp
--- a/Engine/Plugins/Experimental/Animation/ContextualAnimation/Source/ContextualAnimationEditor/Private/ContextualAnimationEditor.cpp
+++ b/Engine/Plugins/Experimental/Animation/ContextualAnimation/Source/ContextualAnimationEditor/Private/ContextualAnimationEditor.cpp
@@ -27,7 +27,8 @@ void FContextualAnimationEditorModule::ShutdownModule()
 	FContextualAnimEditorStyle::Shutdown();
 
 	FPropertyEditorModule& PropertyModule = FModuleManager::LoadModuleChecked<FPropertyEditorModule>("PropertyEditor");
-	PropertyModule.UnregisterCustomPropertyTypeLayout("ContextualAnimFloatParam");
+	PropertyModule.UnregisterCustomPropertyTypeLayout("ContextualAnimDistanceParam");
+	PropertyModule.UnregisterCustomPropertyTypeLayout("ContextualAnimAngleParam");
 
 	FEditorModeRegistry::Get().UnregisterMode(FContextualAnimEdMode::EM_ContextualAnimEdModeId);
 }
diff --git a/Engine/Plugins/Experimental/Animation/ContextualAnimation/Source/ContextualAnimationEditor/Public/ContextualAnimEdMode.h b/Engine/Plugins/Experimental/Animation/ContextualAnimation/Source/ContextualAnimationEditor/Public/ContextualAnimEdMode.h
--- a/Engine/Plugins/Experimental/Animation/ContextualAnimation/Source/ContextualAnimationEditor/Public/ContextualAnimEdMode.h
+++ b/Engine/Plugins/Experimental/Animation/ContextualAnimation/Source/ContextualAnimationEditor/Public/ContextualAnimEdMode.h
@@ -3,6 +3,7 @@
 #pragma once
 
 #include "CoreMinimal.h"
+#include "Templates/SubclassOf.h"
 #include "EdMode.h"
 
 class FContextualAnimEdModeToolkit;
@@ -10,8 +11,9 @@ class FContextualAnimEdModeToolkit;
 class FContextualAnimEdMode : public FEdMode
 {
 public:
+
 	const static FEditorModeID EM_ContextualAnimEdModeId;
-public:
+
 	FContextualAnimEdMode();
 	virtual ~FContextualAnimEdMode();
 
@@ -24,7 +26,14 @@ public:
 	virtual bool UsesToolkits() const override;
 	// End of FEdMode interface
 
-	TSharedPtr<FContextualAnimEdModeToolkit> GetContextualAnimEdModeToolkit();
+	/** FGCObject interface */
+	virtual void AddReferencedObjects(FReferenceCollector& Collector) override;
+
+	static FContextualAnimEdMode& Get();
+
+	TSharedPtr<FContextualAnimEdModeToolkit> GetContextualAnimEdModeToolkit() const;
+
+	class UContextualAnimPreviewManager* PreviewManager;
 
-	TWeakObjectPtr<class ACharacter> TestCharacter;
+	bool GetHitResultUnderCursor(FHitResult& OutHitResult, FEditorViewportClient* InViewportClient, const FViewportClick& Click) const;
 };
diff --git a/Engine/Plugins/Experimental/Animation/ContextualAnimation/Source/ContextualAnimationEditor/Public/ContextualAnimEdModeSettings.h b/Engine/Plugins/Experimental/Animation/ContextualAnimation/Source/ContextualAnimationEditor/Public/ContextualAnimEdModeSettings.h
--- a/Engine/Plugins/Experimental/Animation/ContextualAnimation/Source/ContextualAnimationEditor/Public/ContextualAnimEdModeSettings.h
+++ b/Engine/Plugins/Experimental/Animation/ContextualAnimation/Source/ContextualAnimationEditor/Public/ContextualAnimEdModeSettings.h
@@ -6,7 +6,9 @@
 #include "Templates/SubclassOf.h"
 #include "ContextualAnimEdModeSettings.generated.h"
 
+class AActor;
 class ACharacter;
+class UContextualAnimSceneAsset;
 
 UCLASS()
 class UContextualAnimEdModeSettings : public UObject
@@ -16,7 +18,7 @@ class UContextualAnimEdModeSettings : public UObject
 public:
 
 	UPROPERTY(EditAnywhere, Category = "Settings")
-	TSubclassOf<ACharacter> TestCharacterClass;
+	UContextualAnimSceneAsset* SceneAsset;
 
 	UContextualAnimEdModeSettings(const FObjectInitializer& ObjectInitializer);
 };
\ No newline at end of file
diff --git a/Engine/Plugins/Experimental/Animation/ContextualAnimation/Source/ContextualAnimationEditor/Public/ContextualAnimEdModeToolkit.h b/Engine/Plugins/Experimental/Animation/ContextualAnimation/Source/ContextualAnimationEditor/Public/ContextualAnimEdModeToolkit.h
--- a/Engine/Plugins/Experimental/Animation/ContextualAnimation/Source/ContextualAnimationEditor/Public/ContextualAnimEdModeToolkit.h
+++ b/Engine/Plugins/Experimental/Animation/ContextualAnimation/Source/ContextualAnimationEditor/Public/ContextualAnimEdModeToolkit.h
@@ -35,5 +35,8 @@ private:
 	UContextualAnimEdModeSettings* Settings;
 
 	TSharedPtr<SWidget> ToolkitWidget;
+
 	TSharedPtr<IDetailsView> EdModeSettingsWidget;
+
+	TSharedPtr<IDetailsView> PreviewManagerWidget;
 };
diff --git a/Engine/Plugins/Experimental/Animation/ContextualAnimation/Source/ContextualAnimationEditor/Public/ContextualAnimPreviewManager.h b/Engine/Plugins/Experimental/Animation/ContextualAnimation/Source/ContextualAnimationEditor/Public/ContextualAnimPreviewManager.h
new file mode 100644
--- /dev/null
+++ b/Engine/Plugins/Experimental/Animation/ContextualAnimation/Source/ContextualAnimationEditor/Public/ContextualAnimPreviewManager.h
@@ -0,0 +1,49 @@
+// Copyright Epic Games, Inc. All Rights Reserved.
+
+#pragma once
+
+#include "CoreMinimal.h"
+#include "UObject/Object.h"
+#include "Templates/SubclassOf.h"
+#include "ContextualAnimPreviewManager.generated.h"
+
+class ACharacter;
+class UContextualAnimSceneAsset;
+
+UCLASS()
+class UContextualAnimPreviewManager : public UObject
+{
+	GENERATED_BODY()
+
+public:
+
+	UPROPERTY(EditAnywhere, Category = "Preview")
+	TSubclassOf<ACharacter> DefaultPreviewClass;
+
+	UPROPERTY(EditAnywhere, Category = "Debug")
+	bool bDrawDebugScene;
+
+	UPROPERTY(EditAnywhere, Category = "Debug", meta = (EditCondition = "bDrawDebugScene"))
+	float Time;
+
+	UPROPERTY(EditAnywhere, Category = "Debug", meta = (EditCondition = "bDrawDebugScene"))
+	FTransform ScenePivot;
+
+	UPROPERTY()
+	TWeakObjectPtr<ACharacter> TestCharacter;
+
+	UPROPERTY()
+	TMap<FName, AActor*> PreviewActors;
+
+	UContextualAnimPreviewManager(const FObjectInitializer& ObjectInitializer);
+
+	void SpawnPreviewActors(const UContextualAnimSceneAsset* SceneAsset, const FTransform& SceneOrigin);
+
+	AActor* SpawnPreviewActor(UClass* Class, const FTransform& SpawnTransform) const;
+
+	void MoveForward(float Value);
+
+	void MoveRight(float Value);
+
+	void MoveToLocation(const FVector& GoalLocation);
+};
\ No newline at end of file
