Contextual Anim - Added missing include to fix non-unity build #jira UE-109232 #review trivial #fyi [REDACTED_HANDLE]

diff --git a/Engine/Plugins/Experimental/Animation/ContextualAnimation/Source/ContextualAnimation/Public/ContextualAnimSceneInstance.h b/Engine/Plugins/Experimental/Animation/ContextualAnimation/Source/ContextualAnimation/Public/ContextualAnimSceneInstance.h
--- a/Engine/Plugins/Experimental/Animation/ContextualAnimation/Source/ContextualAnimation/Public/ContextualAnimSceneInstance.h
+++ b/Engine/Plugins/Experimental/Animation/ContextualAnimation/Source/ContextualAnimation/Public/ContextualAnimSceneInstance.h
@@ -4,6 +4,7 @@
 
 #include "CoreMinimal.h"
 #include "UObject/Object.h"
+#include "Animation/AnimNotifies/AnimNotify.h"
 #include "ContextualAnimSceneInstance.generated.h"
 
 struct FContextualAnimData;
