Added MemMark to prevent memory leak from contextual anim scene asset

diff --git a/Engine/Plugins/Experimental/Animation/ContextualAnimation/Source/ContextualAnimation/Private/ContextualAnimUtilities.cpp b/Engine/Plugins/Experimental/Animation/ContextualAnimation/Source/ContextualAnimation/Private/ContextualAnimUtilities.cpp
--- a/Engine/Plugins/Experimental/Animation/ContextualAnimation/Source/ContextualAnimation/Private/ContextualAnimUtilities.cpp
+++ b/Engine/Plugins/Experimental/Animation/ContextualAnimation/Source/ContextualAnimation/Private/ContextualAnimUtilities.cpp
@@ -11,9 +11,12 @@
 #include "Engine/World.h"
 #include "ContextualAnimTypes.h"
 #include "ContextualAnimSceneAsset.h"
+#include "Misc/MemStack.h"
 
 void UContextualAnimUtilities::ExtractLocalSpacePose(const UAnimSequenceBase* Animation, const FBoneContainer& BoneContainer, float Time, bool bExtractRootMotion, FCompactPose& OutPose)
 {
+	FMemMark Mark(FMemStack::Get());
+
 	OutPose.SetBoneContainer(&BoneContainer);
 
 	FBlendedCurve Curve;
