Fix for unity build. Temporarily fixed StringConv.h until Devin can make a proper fix.

diff --git a/Engine/Plugins/Experimental/Animation/ContextualAnimation/Source/ContextualAnimation/Public/ContextualAnimSceneInstance.h b/Engine/Plugins/Experimental/Animation/ContextualAnimation/Source/ContextualAnimation/Public/ContextualAnimSceneInstance.h
--- a/Engine/Plugins/Experimental/Animation/ContextualAnimation/Source/ContextualAnimation/Public/ContextualAnimSceneInstance.h
+++ b/Engine/Plugins/Experimental/Animation/ContextualAnimation/Source/ContextualAnimation/Public/ContextualAnimSceneInstance.h
@@ -5,6 +5,7 @@
 #include "CoreMinimal.h"
 #include "UObject/Object.h"
 #include "Animation/AnimNotifies/AnimNotify.h"
+#include "GameFramework/Actor.h"
 #include "ContextualAnimSceneInstance.generated.h"
 
 struct FContextualAnimData;
diff --git a/Engine/Source/Editor/UnrealEd/Private/Commandlets/DerivedDataCacheCommandlet.cpp b/Engine/Source/Editor/UnrealEd/Private/Commandlets/DerivedDataCacheCommandlet.cpp
--- a/Engine/Source/Editor/UnrealEd/Private/Commandlets/DerivedDataCacheCommandlet.cpp
+++ b/Engine/Source/Editor/UnrealEd/Private/Commandlets/DerivedDataCacheCommandlet.cpp
@@ -23,6 +23,7 @@ DerivedDataCacheCommandlet.cpp: Commandlet for DDC maintenence
 #include "Algo/RemoveIf.h"
 #include "Algo/Transform.h"
 #include "Settings/ProjectPackagingSettings.h"
+#include "Editor.h"
 
 DEFINE_LOG_CATEGORY_STATIC(LogDerivedDataCacheCommandlet, Log, All);
 
