Moved MemMark to a high level scope to prevent a leak causing by the off line generation of the alignment tracks and added some comments

diff --git a/Engine/Plugins/Experimental/Animation/ContextualAnimation/Source/ContextualAnimation/Private/ContextualAnimCompositeSceneAsset.cpp b/Engine/Plugins/Experimental/Animation/ContextualAnimation/Source/ContextualAnimation/Private/ContextualAnimCompositeSceneAsset.cpp
--- a/Engine/Plugins/Experimental/Animation/ContextualAnimation/Source/ContextualAnimation/Private/ContextualAnimCompositeSceneAsset.cpp
+++ b/Engine/Plugins/Experimental/Animation/ContextualAnimation/Source/ContextualAnimation/Private/ContextualAnimCompositeSceneAsset.cpp
@@ -25,8 +25,16 @@ EContextualAnimJoinRule UContextualAnimCompositeSceneAsset::GetJoinRuleForRole(c
 
 void UContextualAnimCompositeSceneAsset::PreSave(const class ITargetPlatform* TargetPlatform)
 {
+	// Necessary for FCompactPose that uses a FAnimStackAllocator (TMemStackAllocator) which allocates from FMemStack.
+	// When allocating memory from FMemStack we need to explicitly use FMemMark to ensure items are freed when the scope exits. 
+	// UWorld::Tick adds a FMemMark to catch any allocation inside the game tick 
+	// but any allocation from outside the game tick (like here when generating the alignment tracks off-line) must explicitly add a mark to avoid a leak 
+	FMemMark Mark(FMemStack::Get());
+
+	// Generate scene pivot for each alignment section
 	Super::PreSave(TargetPlatform);
 
+	// Generate alignment tracks relative to scene pivot
 	for (FContextualAnimData& Data : InteractorTrack.AnimDataContainer)
 	{
 		GenerateAlignmentTracksRelativeToScenePivot(Data);
diff --git a/Engine/Plugins/Experimental/Animation/ContextualAnimation/Source/ContextualAnimation/Private/ContextualAnimSceneAsset.cpp b/Engine/Plugins/Experimental/Animation/ContextualAnimation/Source/ContextualAnimation/Private/ContextualAnimSceneAsset.cpp
--- a/Engine/Plugins/Experimental/Animation/ContextualAnimation/Source/ContextualAnimation/Private/ContextualAnimSceneAsset.cpp
+++ b/Engine/Plugins/Experimental/Animation/ContextualAnimation/Source/ContextualAnimation/Private/ContextualAnimSceneAsset.cpp
@@ -14,6 +14,13 @@ UContextualAnimSceneAsset::UContextualAnimSceneAsset(const FObjectInitializer& O
 
 void UContextualAnimSceneAsset::PreSave(const class ITargetPlatform* TargetPlatform)
 {
+	// Necessary for FCompactPose that uses a FAnimStackAllocator (TMemStackAllocator) which allocates from FMemStack.
+	// When allocating memory from FMemStack we need to explicitly use FMemMark to ensure items are freed when the scope exits. 
+	// UWorld::Tick adds a FMemMark to catch any allocation inside the game tick 
+	// but any allocation from outside the game tick (like here when generating the alignment tracks off-line) must explicitly add a mark to avoid a leak 
+	FMemMark Mark(FMemStack::Get());
+
+	// Generate scene pivot for each alignment section
 	Super::PreSave(TargetPlatform);
 	
 	const FTransform MeshToComponentInverse = MeshToComponent.Inverse();
diff --git a/Engine/Plugins/Experimental/Animation/ContextualAnimation/Source/ContextualAnimation/Private/ContextualAnimUtilities.cpp b/Engine/Plugins/Experimental/Animation/ContextualAnimation/Source/ContextualAnimation/Private/ContextualAnimUtilities.cpp
--- a/Engine/Plugins/Experimental/Animation/ContextualAnimation/Source/ContextualAnimation/Private/ContextualAnimUtilities.cpp
+++ b/Engine/Plugins/Experimental/Animation/ContextualAnimation/Source/ContextualAnimation/Private/ContextualAnimUtilities.cpp
@@ -15,8 +15,6 @@
 
 void UContextualAnimUtilities::ExtractLocalSpacePose(const UAnimSequenceBase* Animation, const FBoneContainer& BoneContainer, float Time, bool bExtractRootMotion, FCompactPose& OutPose)
 {
-	FMemMark Mark(FMemStack::Get());
-
 	OutPose.SetBoneContainer(&BoneContainer);
 
 	FBlendedCurve Curve;
diff --git a/Engine/Plugins/Experimental/Animation/ContextualAnimation/Source/ContextualAnimation/Public/ContextualAnimUtilities.h b/Engine/Plugins/Experimental/Animation/ContextualAnimation/Source/ContextualAnimation/Public/ContextualAnimUtilities.h
--- a/Engine/Plugins/Experimental/Animation/ContextualAnimation/Source/ContextualAnimation/Public/ContextualAnimUtilities.h
+++ b/Engine/Plugins/Experimental/Animation/ContextualAnimation/Source/ContextualAnimation/Public/ContextualAnimUtilities.h
@@ -15,8 +15,20 @@ class CONTEXTUALANIMATION_API UContextualAnimUtilities : public UBlueprintFuncti
 
 public:
 
+	/** 
+	 * Helper function to extract local space pose from an animation at a given time.
+	 * If the supplied animation is a montage it will extract the pose from the first track
+	 * IMPORTANT: This function expects you to add a MemMark (FMemMark Mark(FMemStack::Get());) at the correct scope if you are using it from outside world's tick
+	 */
 	static void ExtractLocalSpacePose(const UAnimSequenceBase* Animation, const FBoneContainer& BoneContainer, float Time, bool bExtractRootMotion, FCompactPose& OutPose);
+
+	/**
+	 * Helper function to extract component space pose from an animation at a given time
+     * If the supplied animation is a montage it will extract the pose from the first track
+	 * IMPORTANT: This function expects you to add a MemMark (FMemMark Mark(FMemStack::Get());) at the correct scope if you are using it from outside world's tick
+	 */
 	static void ExtractComponentSpacePose(const UAnimSequenceBase* Animation, const FBoneContainer& BoneContainer, float Time, bool bExtractRootMotion, FCSPose<FCompactPose>& OutPose);
+
 	static FTransform ExtractRootMotionFromAnimation(const UAnimSequenceBase* Animation, float StartTime, float EndTime);
 
 	static void DrawDebugPose(const UWorld* World, const UAnimSequenceBase* Animation, float Time, const FTransform& LocalToWorldTransform, const FColor& Color, float LifeTime, float Thickness);
