Contextual Anim
- Added support for IK targets including option to autogenerate IK targets from source data.
- Implemented simple state machine to support transitions within an interaction (Experimental)
- Changed the way threshold area is defined from a cone like shape to a constrained polygon with four vertices and moved selection tests to the metadata object
- Added global container to store all SceneActorComps and reparented SceneActorComp
- Removed alignment tracks relative to other actors since it can be calculated at runtime.
- Refactored SceneAsset to contain an CompositeTrack for each role. CompositeSceneAsset will be removed shortly.
- Users can now specify a custom blueprint for the SceneInstance to run game specific logic during the interaction
- Multiple functions and data is now exposed to blueprint

diff --git a/Engine/Plugins/Experimental/Animation/ContextualAnimation/Source/ContextualAnimation/Private/ContextualAnimCompositeSceneAsset.cpp b/Engine/Plugins/Experimental/Animation/ContextualAnimation/Source/ContextualAnimation/Private/ContextualAnimCompositeSceneAsset.cpp
--- a/Engine/Plugins/Experimental/Animation/ContextualAnimation/Source/ContextualAnimation/Private/ContextualAnimCompositeSceneAsset.cpp
+++ b/Engine/Plugins/Experimental/Animation/ContextualAnimation/Source/ContextualAnimation/Private/ContextualAnimCompositeSceneAsset.cpp
@@ -5,6 +5,7 @@
 #include "ContextualAnimMetadata.h"
 #include "Animation/AnimMontage.h"
 #include "UObject/ObjectSaveContext.h"
+#include "ContextualAnimScenePivotProvider.h"
 
 const FName UContextualAnimCompositeSceneAsset::InteractorRoleName = FName(TEXT("interactor"));
 const FName UContextualAnimCompositeSceneAsset::InteractableRoleName = FName(TEXT("interactable"));
@@ -15,21 +16,26 @@ UContextualAnimCompositeSceneAsset::UContextualAnimCompositeSceneAsset(const FOb
 	PrimaryRole = UContextualAnimCompositeSceneAsset::InteractableRoleName;
 }
 
-UClass* UContextualAnimCompositeSceneAsset::GetPreviewActorClassForRole(const FName& Role) const
+const FContextualAnimTrackSettings* UContextualAnimCompositeSceneAsset::GetTrackSettings(const FName& Role) const
 {
-	return Role == PrimaryRole ? InteractableTrack.Settings.PreviewActorClass : InteractorTrack.Settings.PreviewActorClass;
+	return Role == PrimaryRole ? &InteractableTrack.Settings : &InteractorTrack.Settings;
 }
 
-EContextualAnimJoinRule UContextualAnimCompositeSceneAsset::GetJoinRuleForRole(const FName& Role) const
+const FContextualAnimData* UContextualAnimCompositeSceneAsset::GetAnimDataForRoleAtIndex(const FName& Role, int32 Index) const
 {
-	return Role == PrimaryRole ? InteractableTrack.Settings.JoinRule : InteractorTrack.Settings.JoinRule;
-}
+	if(Role == UContextualAnimCompositeSceneAsset::InteractableRoleName)
+	{
+		return &InteractableTrack.AnimData;
+	}
+	else if(Role == UContextualAnimCompositeSceneAsset::InteractorRoleName)
+	{
+		if(InteractorTrack.AnimDataContainer.IsValidIndex(Index))
+		{
+			return &InteractorTrack.AnimDataContainer[Index];
+		}
+	}
 
-void UContextualAnimCompositeSceneAsset::PreSave(const class ITargetPlatform* TargetPlatform)
-{
-	PRAGMA_DISABLE_DEPRECATION_WARNINGS;
-	Super::PreSave(TargetPlatform);
-	PRAGMA_ENABLE_DEPRECATION_WARNINGS;
+	return nullptr;
 }
 
 void UContextualAnimCompositeSceneAsset::PreSave(FObjectPreSaveContext ObjectSaveContext)
@@ -40,131 +46,76 @@ void UContextualAnimCompositeSceneAsset::PreSave(FObjectPreSaveContext ObjectSav
 	// but any allocation from outside the game tick (like here when generating the alignment tracks off-line) must explicitly add a mark to avoid a leak 
 	FMemMark Mark(FMemStack::Get());
 
-	// Generate scene pivot for each alignment section
 	Super::PreSave(ObjectSaveContext);
 
-	// Generate alignment tracks relative to scene pivot
-	for (FContextualAnimData& Data : InteractorTrack.AnimDataContainer)
+	InteractableTrack.AnimData.Index = 0;
+
+	int32 NumAnimData = InteractorTrack.AnimDataContainer.Num();
+	for (int32 AnimDataIdx = 0; AnimDataIdx < NumAnimData; AnimDataIdx++)
 	{
-		GenerateAlignmentTracksRelativeToScenePivot(Data);
+		InteractorTrack.AnimDataContainer[AnimDataIdx].Index = AnimDataIdx;
 	}
-}
-
-bool UContextualAnimCompositeSceneAsset::QueryData(FContextualAnimQueryResult& OutResult, const FContextualAnimQueryParams& QueryParams, const FTransform& ToWorldTransform) const
-{
-	OutResult.Reset();
-
-	const FTransform QueryTransform = QueryParams.Querier.IsValid() ? QueryParams.Querier->GetActorTransform() : QueryParams.QueryTransform;
 
-	int32 DataIndex = INDEX_NONE;
-	if (QueryParams.bComplexQuery)
+	// Generate scene pivot for each alignment section
+	for (int32 AlignmentSectionIdx = 0; AlignmentSectionIdx < AlignmentSections.Num(); AlignmentSectionIdx++)
 	{
-		for (int32 Idx = 0; Idx < InteractorTrack.AnimDataContainer.Num(); Idx++)
-		{
-			const FContextualAnimData& Data = InteractorTrack.AnimDataContainer[Idx];
+		AlignmentSections[AlignmentSectionIdx].ScenePivots.Reset();
 
-			if (Data.Metadata)
+		// We need to calculate scene pivot for each set of animations
+		for (int32 AnimDataIdx = 0; AnimDataIdx < NumAnimData; AnimDataIdx++)
+		{
+			if (AlignmentSections[AlignmentSectionIdx].ScenePivotProvider)
 			{
-				const FTransform EntryTransform = Data.GetAlignmentTransformAtEntryTime() * ToWorldTransform;
-
-				FVector Origin = ToWorldTransform.GetLocation();
-				FVector DirToEntry = (EntryTransform.GetLocation() - Origin).GetSafeNormal2D();
-
-				if (Data.Metadata->OffsetFromOrigin != 0.f)
-				{
-					Origin = Origin + DirToEntry * Data.Metadata->OffsetFromOrigin;
-				}
-
-				// Distance Test
-				//--------------------------------------------------
-				if (Data.Metadata->Distance.MaxDistance > 0.f || Data.Metadata->Distance.MinDistance > 0.f)
-				{
-					const float DistSq = FVector::DistSquared2D(Origin, QueryTransform.GetLocation());
-
-					if (Data.Metadata->Distance.MaxDistance > 0.f)
-					{
-						if (DistSq > FMath::Square(Data.Metadata->Distance.MaxDistance))
-						{
-							continue;
-						}
-					}
-
-					if (Data.Metadata->Distance.MinDistance > 0.f)
-					{
-						if (DistSq < FMath::Square(Data.Metadata->Distance.MinDistance))
-						{
-							continue;
-						}
-					}
-				}
-
-				// Angle Test
-				//--------------------------------------------------
-				if (Data.Metadata->Angle.Tolerance > 0.f)
-				{
-					//@TODO: Cache this
-					const float AngleCos = FMath::Cos(FMath::Clamp(FMath::DegreesToRadians(Data.Metadata->Angle.Tolerance), 0.f, PI));
-					const FVector ToTarget = (QueryTransform.GetLocation() - Origin).GetSafeNormal2D();
-					if (FVector::DotProduct(ToTarget, DirToEntry) < AngleCos)
-					{
-						continue;
-					}
-				}
-
-				// Facing Test
-				//--------------------------------------------------
-				if (Data.Metadata->Facing.Tolerance > 0.f)
-				{
-					//@TODO: Cache this
-					const float FacingCos = FMath::Cos(FMath::Clamp(FMath::DegreesToRadians(Data.Metadata->Facing.Tolerance), 0.f, PI));
-					const FVector ToTarget = (ToWorldTransform.GetLocation() - QueryTransform.GetLocation()).GetSafeNormal2D();
-					if (FVector::DotProduct(QueryTransform.GetRotation().GetForwardVector(), ToTarget) < FacingCos)
-					{
-						continue;
-					}
-				}
+				const FTransform ScenePivot = AlignmentSections[AlignmentSectionIdx].ScenePivotProvider->CalculateScenePivot_Source(AnimDataIdx);
+				AlignmentSections[AlignmentSectionIdx].ScenePivots.Add(ScenePivot);
 			}
-
-			// Return the first item that passes all tests
-			DataIndex = Idx;
-			break;
-		}
-	}
-	else // Simple Query
-	{
-		float BestDistanceSq = MAX_FLT;
-		for (int32 Idx = 0; Idx < InteractorTrack.AnimDataContainer.Num(); Idx++)
-		{
-			const FContextualAnimData& Data = InteractorTrack.AnimDataContainer[Idx];
-
-			//@TODO: Convert querier location to local space instead
-			const FTransform EntryTransform = Data.GetAlignmentTransformAtEntryTime() * ToWorldTransform;
-			const float DistSq = FVector::DistSquared2D(EntryTransform.GetLocation(), QueryTransform.GetLocation());
-			if (DistSq < BestDistanceSq)
+			else
 			{
-				BestDistanceSq = DistSq;
-				DataIndex = Idx;
+				AlignmentSections[AlignmentSectionIdx].ScenePivots.Add(FTransform::Identity);
 			}
 		}
 	}
 
-	if (DataIndex != INDEX_NONE)
+	for (FContextualAnimData& Data : InteractorTrack.AnimDataContainer)
 	{
-		const FContextualAnimData& ResultData = InteractorTrack.AnimDataContainer[DataIndex];
+		// Generate alignment tracks relative to scene pivot
+		GenerateAlignmentTracks(InteractorTrack.Settings, Data);
 
-		OutResult.DataIndex = DataIndex;
-		OutResult.Animation = ResultData.Animation;
-		OutResult.EntryTransform = ResultData.GetAlignmentTransformAtEntryTime() * ToWorldTransform;
-		OutResult.SyncTransform = ResultData.GetAlignmentTransformAtSyncTime() * ToWorldTransform;
+		// Generate IK Targets
+		GenerateIKTargetTracks(InteractorTrack.Settings, Data);
+	}
 
-		if (QueryParams.bFindAnimStartTime)
+	UpdateRadius();
+}
+
+bool UContextualAnimCompositeSceneAsset::Query(const FName& Role, FContextualAnimQueryResult& OutResult, const FContextualAnimQueryParams& QueryParams, const FTransform& ToWorldTransform) const
+{
+	//@TODO: Intentionally ignoring Role param since it doesn't make sense for this asset. This is just a temp anyway until we remove this asset completely
+	return QueryCompositeTrack(&InteractorTrack, OutResult, QueryParams, ToWorldTransform);
+}
+
+bool UContextualAnimCompositeSceneAsset::QueryData(FContextualAnimQueryResult& OutResult, const FContextualAnimQueryParams& QueryParams, const FTransform& ToWorldTransform) const
+{
+	return QueryCompositeTrack(&InteractorTrack, OutResult, QueryParams, ToWorldTransform);
+}
+
+void UContextualAnimCompositeSceneAsset::ForEachAnimData(FForEachAnimDataFunction Function) const
+{
+	if (Function(InteractableRoleName, InteractableTrack.AnimData) == EContextualAnimForEachResult::Break)
+	{
+		return;
+	}
+
+	for (const FContextualAnimData& AnimData : InteractorTrack.AnimDataContainer)
+	{
+		if (Function(InteractorRoleName, AnimData) == EContextualAnimForEachResult::Break)
 		{
-			const FVector LocalLocation = (QueryTransform.GetRelativeTransform(ToWorldTransform)).GetLocation();
-			OutResult.AnimStartTime = ResultData.FindBestAnimStartTime(LocalLocation);
+			return;
 		}
-
-		return true;
 	}
+}
 
-	return false;
+TArray<FName> UContextualAnimCompositeSceneAsset::GetRoles() const 
+{ 
+	return { InteractableRoleName, InteractorRoleName }; 
 }
\ No newline at end of file
diff --git a/Engine/Plugins/Experimental/Animation/ContextualAnimation/Source/ContextualAnimation/Private/ContextualAnimManager.cpp b/Engine/Plugins/Experimental/Animation/ContextualAnimation/Source/ContextualAnimation/Private/ContextualAnimManager.cpp
--- a/Engine/Plugins/Experimental/Animation/ContextualAnimation/Source/ContextualAnimation/Private/ContextualAnimManager.cpp
+++ b/Engine/Plugins/Experimental/Animation/ContextualAnimation/Source/ContextualAnimation/Private/ContextualAnimManager.cpp
@@ -1,19 +1,21 @@
 // Copyright Epic Games, Inc. All Rights Reserved.
 
 #include "ContextualAnimManager.h"
-#include "ContextualAnimTypes.h"
 #include "ContextualAnimActorInterface.h"
 #include "ContextualAnimSceneInstance.h"
 #include "DrawDebugHelpers.h"
 #include "ContextualAnimSceneAsset.h"
 #include "ContextualAnimCompositeSceneAsset.h"
 #include "ContextualAnimation.h"
+#include "ContextualAnimSceneActorComponent.h"
 #include "Engine/World.h"
+#include "GameFramework/Actor.h"
+
+DECLARE_CYCLE_STAT(TEXT("ContextualAnim FindClosestSceneActorComp"), STAT_ContextualAnim_FindClosestSceneActorComp, STATGROUP_Anim);
 
 UContextualAnimManager::UContextualAnimManager(const FObjectInitializer& ObjectInitializer)
 	: Super(ObjectInitializer)
 {
-
 }
 
 UContextualAnimManager* UContextualAnimManager::Get(const UWorld* World)
@@ -21,11 +23,74 @@ UContextualAnimManager* UContextualAnimManager::Get(const UWorld* World)
 	return (World) ? FContextualAnimationModule::GetManager(World) : nullptr;
 }
 
+UContextualAnimManager* UContextualAnimManager::GetContextualAnimManager(UObject* WorldContextObject)
+{
+	return Get(GEngine->GetWorldFromContextObject(WorldContextObject, EGetWorldErrorMode::LogAndReturnNull));
+}
+
 UWorld* UContextualAnimManager::GetWorld()const
 {
 	return CastChecked<UWorld>(GetOuter());
 }
 
+ETickableTickType UContextualAnimManager::GetTickableTickType() const
+{
+	//@TODO: Switch to Conditional and use IsTickable to determine whether to tick. It should only tick when scene instances are active
+	return (HasAnyFlags(RF_ClassDefaultObject)) ? ETickableTickType::Never : ETickableTickType::Always;
+}
+TStatId UContextualAnimManager::GetStatId() const
+{
+	RETURN_QUICK_DECLARE_CYCLE_STAT(UContextualAnimManager, STATGROUP_Tickables);
+}
+
+void UContextualAnimManager::RegisterSceneActorComponent(UContextualAnimSceneActorComponent* SceneActorComp)
+{
+	if(SceneActorComp)
+	{
+		SceneActorCompContainer.Add(SceneActorComp);
+	}
+}
+
+void UContextualAnimManager::UnregisterSceneActorComponent(UContextualAnimSceneActorComponent* SceneActorComp)
+{
+	if (SceneActorComp)
+	{
+		SceneActorCompContainer.Remove(SceneActorComp);
+	}
+}
+
+UContextualAnimSceneActorComponent* UContextualAnimManager::FindClosestSceneActorCompToActor(const AActor* Actor) const
+{
+	SCOPE_CYCLE_COUNTER(STAT_ContextualAnim_FindClosestSceneActorComp);
+
+	//@TODO: Implement some space partitioning to reduce the cost of this search (a 2D grid might be enough)
+
+	UContextualAnimSceneActorComponent* ClosestSceneActorComp = nullptr;
+	if(Actor)
+	{
+		const FVector TestLocation = Actor->GetActorLocation();
+
+		float BestDistanceSq = MAX_FLT;
+		for (UContextualAnimSceneActorComponent* SceneActorComp : SceneActorCompContainer)
+		{
+			check(SceneActorComp);
+			check(SceneActorComp->GetOwner());
+
+			if(SceneActorComp->GetOwner() != Actor)
+			{
+				const float DistanceSq = FVector::DistSquared(SceneActorComp->GetComponentLocation(), TestLocation);
+				if (DistanceSq < BestDistanceSq)
+				{
+					BestDistanceSq = DistanceSq;
+					ClosestSceneActorComp = SceneActorComp;
+				}
+			}
+		}
+	}
+
+	return ClosestSceneActorComp;
+}
+
 void UContextualAnimManager::Tick(float DeltaTime)
 {
 	for (UContextualAnimSceneInstance* SceneInstance : Instances)
@@ -66,7 +131,7 @@ UContextualAnimSceneInstance* UContextualAnimManager::GetSceneWithActor(AActor*
 	return nullptr;
 }
 
-bool UContextualAnimManager::TryStartScene(const UContextualAnimSceneAssetBase* SceneAsset, const TMap<FName, AActor*>& Bindings)
+bool UContextualAnimManager::TryStartScene(const UContextualAnimSceneAssetBase* SceneAsset, const FContextualAnimSceneBindings& Bindings)
 {
 	if (SceneAsset == nullptr)
 	{
@@ -75,89 +140,60 @@ bool UContextualAnimManager::TryStartScene(const UContextualAnimSceneAssetBase*
 	}
 
 	TMap<FName, FContextualAnimSceneActorData> SceneActorMap;
-	SceneActorMap.Reserve(Bindings.Num());
+	SceneActorMap.Reserve(Bindings.RoleToActorMap.Num());
 
-	// @TODO: Add virtual function in the asset class that resolve the bindings for us
-	if (const UContextualAnimSceneAsset* SimpleScene = Cast<const UContextualAnimSceneAsset>(SceneAsset))
+	// Find primary actor first so we can perform queries relative to it
+	AActor* const* PrimaryActorPtr = Bindings.RoleToActorMap.Find(SceneAsset->PrimaryRole);
+	if (PrimaryActorPtr == nullptr)
 	{
-		for (const auto& Pair : Bindings)
-		{
-			FName RoleToBind = Pair.Key;
-			AActor* ActorToBind = Pair.Value;
+		UE_LOG(LogContextualAnim, Warning, TEXT("UContextualAnimManager::TryStartScene. Can't start scene. Reason: Attempting to bind an invalid actor to the primary role. SceneAsset: %s Role: %s"),
+			*GetNameSafe(SceneAsset), *SceneAsset->PrimaryRole.ToString());
 
-			if(ActorToBind == nullptr)
-			{
-				UE_LOG(LogTemp, Warning, TEXT("UContextualAnimManager::TryStartScene. Can't start scene. Reason: Attempting to bind an invalid actor. SceneAsset: %s Role: %s"), 
-				*GetNameSafe(SceneAsset), *RoleToBind.ToString());
-
-				return false;
-			}
-
-			const FContextualAnimTrack* Track = SimpleScene->FindTrack(RoleToBind);
-			if(Track == nullptr)
-			{
-				UE_LOG(LogContextualAnim, Warning, TEXT("UContextualAnimManager::TryStartScene. Can't start scene. Reason: Invalid Track. SceneAsset: %s Role: %s"),
-				*GetNameSafe(SceneAsset), *RoleToBind.ToString());
+		return false;
+	}
 
-				return false;
-			}
+	AActor* PrimaryActor = *PrimaryActorPtr;
+	const UContextualAnimSceneActorComponent* PrimarySceneActorComp = PrimaryActor->FindComponentByClass<UContextualAnimSceneActorComponent>();
+	const FTransform ToWorldTransform = PrimarySceneActorComp ? PrimarySceneActorComp->GetComponentTransform() : PrimaryActor->GetActorTransform();
 
-			SceneActorMap.Add(RoleToBind, FContextualAnimSceneActorData(ActorToBind, &Track->AnimData));
-		}
-	}
-	else if (const UContextualAnimCompositeSceneAsset* CompositeScene = Cast<const UContextualAnimCompositeSceneAsset>(SceneAsset))
+	for (const auto& Pair : Bindings.RoleToActorMap)
 	{
-		// Bind the Primary Role first so we can perform the query relative to it
-		FContextualAnimSceneActorData PrimaryData;
-		for (const auto& Pair : Bindings)
+		FName RoleToBind = Pair.Key;
+		AActor* ActorToBind = Pair.Value;
+
+		if (ActorToBind == nullptr)
 		{
-			if (Pair.Key == CompositeScene->PrimaryRole)
-			{
-				PrimaryData.Actor = Pair.Value;
-				PrimaryData.AnimDataPtr = &CompositeScene->InteractableTrack.AnimData;
+			UE_LOG(LogTemp, Warning, TEXT("UContextualAnimManager::TryStartScene. Can't start scene. Reason: Attempting to bind an invalid actor. SceneAsset: %s Role: %s"),
+				*GetNameSafe(SceneAsset), *RoleToBind.ToString());
 
-				if (!PrimaryData.Actor.IsValid())
-				{
-					UE_LOG(LogContextualAnim, Warning, TEXT("UContextualAnimManager::TryStartScene. Can't start scene. Reason: Attempting to bind an invalid actor to the primary role. SceneAsset: %s Role: %s"),
-						*GetNameSafe(SceneAsset), *CompositeScene->PrimaryRole.ToString());
+			return false;
+		}
 
-					return false;
-				}
+		// Primary role is already bound
+		if (RoleToBind != SceneAsset->PrimaryRole)
+		{
+			int32 AnimDataIndex = INDEX_NONE;
+			float AnimStartTime = 0.f;
 
-				if (PrimaryData.AnimDataPtr == nullptr)
+			// Attempt to use AnimDataIndex and AnimStartTime supplied with the bindings
+			if (Bindings.AnimDataIndex != INDEX_NONE)
+			{
+				if (SceneAsset->GetAnimDataForRoleAtIndex(RoleToBind, Bindings.AnimDataIndex) == nullptr)
 				{
-					UE_LOG(LogContextualAnim, Warning, TEXT("UContextualAnimManager::TryStartScene. Can't start scene. Reason: Invalid Track. SceneAsset: %s Role: %s"),
-						*GetNameSafe(SceneAsset), *CompositeScene->PrimaryRole.ToString());
+					UE_LOG(LogTemp, Warning, TEXT("UContextualAnimManager::TryStartScene. Can't start scene. Reason: Invalid AnimDataIndex. SceneAsset: %s Role: %s AnimDataIndex: %d"),
+						*GetNameSafe(SceneAsset), *RoleToBind.ToString(), Bindings.AnimDataIndex);
 
 					return false;
 				}
 
-				SceneActorMap.Add(Pair.Key, PrimaryData);
-
-				break;
-			}
-		}
-
-		for (const auto& Pair : Bindings)
-		{
-			FName RoleToBind = Pair.Key;
-			AActor* ActorToBind = Pair.Value;
-
-			if (ActorToBind == nullptr)
-			{
-				UE_LOG(LogTemp, Warning, TEXT("UContextualAnimManager::TryStartScene. Can't start scene. Reason: Attempting to bind an invalid actor. SceneAsset: %s Role: %s"),
-					*GetNameSafe(SceneAsset), *RoleToBind.ToString());
-
-				return false;
+				AnimDataIndex = Bindings.AnimDataIndex;
+				AnimStartTime = Bindings.AnimStartTime;
 			}
-
-			// Primary role was already bound
-			if (RoleToBind != CompositeScene->PrimaryRole)
+			// Query for best animation to use
+			else
 			{
-				FTransform ToWorldTransform = PrimaryData.GetTransform();
-
 				FContextualAnimQueryResult Result;
-				const bool bQueryResult = CompositeScene->QueryData(Result, FContextualAnimQueryParams(ActorToBind, true, true), ToWorldTransform);
+				const bool bQueryResult = SceneAsset->Query(RoleToBind, Result, FContextualAnimQueryParams(ActorToBind, true, true), ToWorldTransform);
 				if (bQueryResult == false)
 				{
 					UE_LOG(LogContextualAnim, Warning, TEXT("UContextualAnimManager::TryStartScene. Can't start scene. Reason: Can't find Track for actor. SceneAsset: %s Role: %s Actor: %s"),
@@ -166,12 +202,26 @@ bool UContextualAnimManager::TryStartScene(const UContextualAnimSceneAssetBase*
 					return false;
 				}
 
-				SceneActorMap.Add(RoleToBind, FContextualAnimSceneActorData(ActorToBind, &CompositeScene->InteractorTrack.AnimDataContainer[Result.DataIndex], Result.AnimStartTime));
+				AnimDataIndex = Result.DataIndex;
+				AnimStartTime = Result.AnimStartTime;
+			}
+
+			//@TODO: Fix this for +2 actors interactions. AnimStartTime must be calculated only once and then everyone should use it.
+
+			SceneActorMap.Add(RoleToBind, FContextualAnimSceneActorData(ActorToBind, SceneAsset->GetAnimDataForRoleAtIndex(RoleToBind, AnimDataIndex), SceneAsset->GetTrackSettings(RoleToBind), AnimStartTime));
+
+			if (!SceneActorMap.Contains(SceneAsset->PrimaryRole))
+			{
+				SceneActorMap.Add(SceneAsset->PrimaryRole, FContextualAnimSceneActorData(PrimaryActor,
+					SceneAsset->GetAnimDataForRoleAtIndex(SceneAsset->PrimaryRole, AnimDataIndex), 
+					SceneAsset->GetTrackSettings(SceneAsset->PrimaryRole), AnimStartTime)
+				);
 			}
 		}
 	}
 
-	UContextualAnimSceneInstance* NewInstance = NewObject<UContextualAnimSceneInstance>(this);
+	UClass* Class = SceneAsset->SceneInstanceClass;
+	UContextualAnimSceneInstance* NewInstance = Class ? NewObject<UContextualAnimSceneInstance>(this, Class) : NewObject<UContextualAnimSceneInstance>(this);
 	NewInstance->SceneAsset = SceneAsset;
 	NewInstance->SceneActorMap = MoveTemp(SceneActorMap);
 	NewInstance->Start();
@@ -182,6 +232,66 @@ bool UContextualAnimManager::TryStartScene(const UContextualAnimSceneAssetBase*
 	return true;
 }
 
+bool UContextualAnimManager::TryStartScene(const UContextualAnimSceneAssetBase* SceneAsset, AActor* PrimaryActor, const TSet<UContextualAnimSceneActorComponent*>& SceneActorComps)
+{
+	if (SceneAsset == nullptr)
+	{
+		UE_LOG(LogContextualAnim, Warning, TEXT("UContextualAnimManager::TryStartScene. Can't start scene. Reason: Invalid Scene Asset"));
+		return false;
+	}
+
+	if (PrimaryActor == nullptr || IsActorInAnyScene(PrimaryActor))
+	{
+		UE_LOG(LogContextualAnim, Warning, TEXT("UContextualAnimManager::TryStartScene. Can't start scene. Reason: Invalid PrimaryActor or already in another scene. PrimaryActor: %s"), *GetNameSafe(PrimaryActor));
+		return false;
+	}
+
+	const TArray<FName>& Roles = SceneAsset->GetRoles();
+
+	FContextualAnimSceneBindings Bindings;
+	Bindings.RoleToActorMap.Add(SceneAsset->PrimaryRole, PrimaryActor);
+
+	const UContextualAnimSceneActorComponent* PrimarySceneActorComp = PrimaryActor->FindComponentByClass<UContextualAnimSceneActorComponent>();
+	const FTransform ToWorldTransform = PrimarySceneActorComp ? PrimarySceneActorComp->GetComponentTransform() : PrimaryActor->GetActorTransform();
+
+	for (UContextualAnimSceneActorComponent* SceneActorComp : SceneActorComps)
+	{
+		if (SceneActorComp && SceneActorComp->GetOwner())
+		{
+			AActor* ActorToBind = SceneActorComp->GetOwner();
+			if(PrimaryActor != ActorToBind && !IsActorInAnyScene(ActorToBind))
+			{
+				for (const FName& RoleToBind : Roles)
+				{
+					if (!Bindings.RoleToActorMap.Contains(RoleToBind))
+					{
+						FContextualAnimQueryResult Result;
+						if (SceneAsset->Query(RoleToBind, Result, FContextualAnimQueryParams(ActorToBind, true, true), ToWorldTransform))
+						{
+							Bindings.RoleToActorMap.Add(RoleToBind, ActorToBind);
+
+							if (Bindings.AnimDataIndex != INDEX_NONE)
+							{
+								Bindings.AnimDataIndex = Result.DataIndex;
+								Bindings.AnimStartTime = Result.AnimStartTime;
+							}
+
+							break;
+						}
+					}
+				}
+
+				if (Roles.Num() == Bindings.RoleToActorMap.Num())
+				{
+					return TryStartScene(SceneAsset, Bindings);
+				}
+			}
+		}
+	}
+
+	return false;
+}
+
 bool UContextualAnimManager::TryStopSceneWithActor(AActor* Actor)
 {
 	if(UContextualAnimSceneInstance* SceneInstance = GetSceneWithActor(Actor))
diff --git a/Engine/Plugins/Experimental/Animation/ContextualAnimation/Source/ContextualAnimation/Private/ContextualAnimMetadata.cpp b/Engine/Plugins/Experimental/Animation/ContextualAnimation/Source/ContextualAnimation/Private/ContextualAnimMetadata.cpp
--- a/Engine/Plugins/Experimental/Animation/ContextualAnimation/Source/ContextualAnimation/Private/ContextualAnimMetadata.cpp
+++ b/Engine/Plugins/Experimental/Animation/ContextualAnimation/Source/ContextualAnimation/Private/ContextualAnimMetadata.cpp
@@ -4,14 +4,113 @@
 #include "ContextualAnimSceneAsset.h"
 #include "Engine/World.h"
 #include "DrawDebugHelpers.h"
+#include "UObject/ObjectSaveContext.h"
+
+// Copied from GeomTools.cpp and changed second param to be an ArrayView
+static bool IsPointInPolygon(const FVector2D& TestPoint, const TArrayView<FVector2D>& PolygonPoints)
+{
+	const int NumPoints = PolygonPoints.Num();
+	float AngleSum = 0.0f;
+	for (int PointIndex = 0; PointIndex < NumPoints; ++PointIndex)
+	{
+		const FVector2D& VecAB = PolygonPoints[PointIndex] - TestPoint;
+		const FVector2D& VecAC = PolygonPoints[(PointIndex + 1) % NumPoints] - TestPoint;
+		const float Angle = FMath::Sign(FVector2D::CrossProduct(VecAB, VecAC)) * FMath::Acos(FMath::Clamp(FVector2D::DotProduct(VecAB, VecAC) / (VecAB.Size() * VecAC.Size()), -1.0f, 1.0f));
+		AngleSum += Angle;
+	}
+	return (FMath::Abs(AngleSum) > 0.001f);
+}
 
 UContextualAnimMetadata::UContextualAnimMetadata(const FObjectInitializer& ObjectInitializer)
 	: Super(ObjectInitializer)
 {
-
 }
 
 UContextualAnimSceneAssetBase* UContextualAnimMetadata::GetSceneAssetOwner() const
 {
 	return Cast<UContextualAnimSceneAssetBase>(GetOuter());
+}
+
+bool UContextualAnimMetadata::DoesQuerierPassConditions(const FContextualAnimQuerier& Querier, const FContextualAnimQueryContext& Context, const FTransform& EntryTransform) const
+{
+	//@TODO: EntryTransform can be get from here by accessing the AnimData that owns us
+
+	const FTransform ToWorldTransform = Context.Actor.IsValid() ? Context.Actor->GetActorTransform() : Context.Transform;
+	const FTransform QueryTransform = Querier.Actor.IsValid() ? Querier.Actor->GetActorTransform() : Querier.Transform;
+
+	// Facing Test
+	//--------------------------------------------------
+	if (Facing > 0.f)
+	{
+		//@TODO: Cache this
+		const float FacingCos = FMath::Cos(FMath::Clamp(FMath::DegreesToRadians(Facing), 0.f, PI));
+		const FVector ToTarget = (ToWorldTransform.GetLocation() - QueryTransform.GetLocation()).GetSafeNormal2D();
+		if (FVector::DotProduct(QueryTransform.GetRotation().GetForwardVector(), ToTarget) < FacingCos)
+		{
+			return false;
+		}
+	}
+
+	// Sector Test
+	//--------------------------------------------------
+
+	if (MaxDistance > 0.f)
+	{
+		FVector Origin = ToWorldTransform.GetLocation();
+		FVector Direction = (EntryTransform.GetLocation() - ToWorldTransform.GetLocation()).GetSafeNormal2D();
+
+		if (DirectionOffset != 0.f)
+		{
+			Direction = Direction.RotateAngleAxis(DirectionOffset, FVector::UpVector);
+		}
+
+		if (OriginOffset.X != 0.f)
+		{
+			Origin = Origin + Direction * OriginOffset.X;
+		}
+
+		if (OriginOffset.Y != 0.f)
+		{
+			Origin = Origin + (Direction.ToOrientationQuat().GetRightVector()) * OriginOffset.Y;
+		}
+
+		if (NearWidth > 0.f || FarWidth > 0.f)
+		{
+			const float HalfNearWidth = (NearWidth / 2.f);
+			const FVector RightVector = Direction.ToOrientationQuat().GetRightVector();
+			const FVector A = Origin + (-RightVector * HalfNearWidth);
+			const FVector B = Origin + (RightVector * HalfNearWidth);
+
+			const float HalfFarWidth = (FarWidth / 2.f);
+			const FVector FarEdgeCenter = Origin + (Direction * MaxDistance);
+			const FVector C = FarEdgeCenter + (-RightVector * HalfFarWidth);
+			const FVector D = FarEdgeCenter + (RightVector * HalfFarWidth);
+
+			const float DistSq = FVector::DistSquared2D(Origin, QueryTransform.GetLocation());
+			if (DistSq > FMath::Square((D - Origin).Size2D()))
+			{
+				return false;
+			}
+
+			// @TODO: Cache this if the owner is static 
+			// or could we cache it in local space and perform the query in local space?
+			TArray<FVector2D, TInlineAllocator<4>> PolygonPoints = { FVector2D(A), FVector2D(C), FVector2D(D), FVector2D(B) };
+
+			FVector2D TestPoint = FVector2D(QueryTransform.GetLocation());
+			if (!IsPointInPolygon(TestPoint, PolygonPoints))
+			{
+				return false;
+			}
+		}
+		else
+		{
+			const float DistSq = FVector::DistSquared2D(Origin, QueryTransform.GetLocation());
+			if (DistSq > FMath::Square(MaxDistance))
+			{
+				return false;
+			}
+		}
+	}
+
+	return true;
 }
\ No newline at end of file
diff --git a/Engine/Plugins/Experimental/Animation/ContextualAnimation/Source/ContextualAnimation/Private/ContextualAnimSceneActorComponent.cpp b/Engine/Plugins/Experimental/Animation/ContextualAnimation/Source/ContextualAnimation/Private/ContextualAnimSceneActorComponent.cpp
--- a/Engine/Plugins/Experimental/Animation/ContextualAnimation/Source/ContextualAnimation/Private/ContextualAnimSceneActorComponent.cpp
+++ b/Engine/Plugins/Experimental/Animation/ContextualAnimation/Source/ContextualAnimation/Private/ContextualAnimSceneActorComponent.cpp
@@ -4,37 +4,141 @@
 #include "Animation/AnimInstance.h"
 #include "GameFramework/Character.h"
 #include "ContextualAnimMetadata.h"
+#include "Components/SkeletalMeshComponent.h"
+#include "DrawDebugHelpers.h"
+#include "ContextualAnimManager.h"
+#include "ContextualAnimSceneInstance.h"
+#include "ContextualAnimCompositeSceneAsset.h"
+#include "ContextualAnimSceneAsset.h"
+#include "ContextualAnimUtilities.h"
 
 UContextualAnimSceneActorComponent::UContextualAnimSceneActorComponent(const FObjectInitializer& ObjectInitializer)
 	: Super(ObjectInitializer)
 {
-	PrimaryComponentTick.bCanEverTick = true;
+	PrimaryComponentTick.bCanEverTick = false;
 	PrimaryComponentTick.bStartWithTickEnabled = false;
+}
 
-	SetCollisionEnabled(ECollisionEnabled::QueryOnly);
-	SetCollisionResponseToAllChannels(ECR_Ignore);
-	SetCollisionResponseToChannel(ECC_Pawn, ECR_Overlap);
-	SetGenerateOverlapEvents(true);
+FBoxSphereBounds UContextualAnimSceneActorComponent::CalcBounds(const FTransform& LocalToWorld) const
+{
+	const float Radius = SceneAsset ? SceneAsset->GetRadius() : 0.f;
+	return FBoxSphereBounds(FSphere(GetComponentTransform().GetLocation(), Radius));
+}
 
-	bHiddenInGame = false;
+void UContextualAnimSceneActorComponent::OnRegister()
+{
+	Super::OnRegister();
+	UContextualAnimManager::Get(GetWorld())->RegisterSceneActorComponent(this);
 }
 
-bool UContextualAnimSceneActorComponent::QueryData(const FContextualAnimQueryParams& QueryParams, FContextualAnimQueryResult& Result) const
+void UContextualAnimSceneActorComponent::OnUnregister()
 {
-	return SceneAsset ? SceneAsset->QueryData(Result, QueryParams, GetComponentTransform()) : false;
+	Super::OnUnregister();
+	UContextualAnimManager::Get(GetWorld())->UnregisterSceneActorComponent(this);
 }
 
-FBoxSphereBounds UContextualAnimSceneActorComponent::CalcBounds(const FTransform& LocalToWorld) const
+void UContextualAnimSceneActorComponent::OnJoinedScene(const FContextualAnimSceneActorData* SceneActorData)
 {
-	FBox BB(FVector(0.f, 0.f, -10.f), FVector(0.f, 0.f, 10.f));
-	if (SceneAsset)
+	check(SceneActorData);
+	SceneActorDataPtr = SceneActorData;
+
+	// Start listening to TickPose if we joined an scene where we need IK
+	if (SceneActorData->GetSettings()->IKTargetDefinitions.Num() > 0)
 	{
-		for (const FContextualAnimData& Data : SceneAsset->InteractorTrack.AnimDataContainer)
+		if (USkeletalMeshComponent* SkelMeshComp = UContextualAnimUtilities::TryGetSkeletalMeshComponent(GetOwner()))
 		{
-			BB += Data.GetAlignmentTransformAtEntryTime().GetLocation();
+			SkelMeshComp->OnTickPose.AddUObject(this, &UContextualAnimSceneActorComponent::OnTickPose);
+		}
+	}
+
+	OnJoinedSceneDelegate.Broadcast(this);
+}
+
+void UContextualAnimSceneActorComponent::OnLeftScene(const FContextualAnimSceneActorData* SceneActorData)
+{
+	check(SceneActorData);
+
+	OnLeftSceneDelegate.Broadcast(this);
+
+	// Stop listening to TickPose if we were
+	if (SceneActorData->GetSettings()->IKTargetDefinitions.Num() > 0)
+	{
+		if (USkeletalMeshComponent* SkelMeshComp = UContextualAnimUtilities::TryGetSkeletalMeshComponent(GetOwner()))
+		{
+			SkelMeshComp->OnTickPose.RemoveAll(this);
+		}
+	}
+
+	SceneActorDataPtr = nullptr;
+}
+
+void UContextualAnimSceneActorComponent::OnTickPose(class USkinnedMeshComponent* SkinnedMeshComponent, float DeltaTime, bool bNeedsValidRootMotion)
+{
+	//@TODO: Check for LOD to prevent this update if the actor is too far away
+	UpdateIKTargets();
+}
+
+void UContextualAnimSceneActorComponent::UpdateIKTargets()
+{
+	const UContextualAnimSceneInstance* SceneInstance = SceneActorDataPtr ? SceneActorDataPtr->GetSceneInstance() : nullptr;
+	if (SceneInstance)
+	{
+		IKTargets.Reset();
+
+		for (const FContextualAnimIKTargetDefinition& IKTargetDef : SceneActorDataPtr->GetSettings()->IKTargetDefinitions)
+		{
+			if (UAnimInstance* AnimInstance = SceneActorDataPtr->GetAnimInstance())
+			{
+				FTransform IKTargetTransform = FTransform::Identity;
+
+				float Alpha = 0.f;
+				bool bCurveFound = AnimInstance->GetCurveValue(IKTargetDef.AlphaCurveName, Alpha);
+
+				//@TODO: IKTargetTransform will be off by 1 frame if we tick before target until setup tick dependencies
+
+				if (IKTargetDef.Provider == EContextualAnimIKTargetProvider::Autogenerated)
+				{
+					if (const FContextualAnimSceneActorData* TargetSceneActorData = SceneInstance->FindSceneActorDataForRole(IKTargetDef.AutoParams.TargetRole))
+					{
+						if (USkeletalMeshComponent* TargetSkelMeshComp = TargetSceneActorData->GetSkeletalMeshComponent())
+						{
+							const FTransform IKTargetParentTransform = TargetSkelMeshComp->GetSocketTransform(IKTargetDef.AutoParams.BoneName);
+
+							const float Time = SceneActorDataPtr->GetAnimTime();
+							IKTargetTransform = SceneActorDataPtr->GetAnimData()->IKTargetData.ExtractTransformAtTime(IKTargetDef.IKGoalName, Time) * IKTargetParentTransform;
+						}
+					}
+				}
+				else if (IKTargetDef.Provider == EContextualAnimIKTargetProvider::Bone)
+				{
+					if (const FContextualAnimSceneActorData* TargetSceneActorData = SceneInstance->FindSceneActorDataForRole(IKTargetDef.BoneParams.TargetRole))
+					{
+						if (USkeletalMeshComponent* TargetSkelMeshComp = TargetSceneActorData->GetSkeletalMeshComponent())
+						{
+							IKTargetTransform = TargetSkelMeshComp->GetSocketTransform(IKTargetDef.BoneParams.BoneName);
+						}
+					}
+				}
+
+				// Draw Debug IK Targets
+#if !(UE_BUILD_SHIPPING || UE_BUILD_TEST)
+				if (FContextualAnimCVars::CVarDrawDebugIKTargets.GetValueOnGameThread() > 0)
+				{
+					const FColor DrawColor = FColor::MakeRedToGreenColorFromScalar(Alpha);
+					DrawDebugSphere(GetWorld(), IKTargetTransform.GetLocation(), 5.f, 12, DrawColor, false, 0.f, 0, 0.5f);
+					DrawDebugCoordinateSystem(GetWorld(), IKTargetTransform.GetLocation(), IKTargetTransform.Rotator(), 10.f, false, 0.f, 0, 0.5f);
+
+					//DrawDebugString(GetWorld(), IKTargetTransform.GetLocation(), FString::Printf(TEXT("%s (%f)"), *IKTargetDef.AlphaCurveName.ToString(), Alpha));
+				}
+#endif
+
+				// Convert IK Target to mesh space
+				IKTargetTransform = IKTargetTransform.GetRelativeTransform(SceneActorDataPtr->GetSkeletalMeshComponent()->GetComponentTransform());
+
+				IKTargets.Add(FContextualAnimIKTarget(IKTargetDef.IKGoalName, Alpha, IKTargetTransform));
+			}
 		}
 	}
-	return FBoxSphereBounds(BB.TransformBy(GetComponentTransform()));
 }
 
 FPrimitiveSceneProxy* UContextualAnimSceneActorComponent::CreateSceneProxy()
@@ -53,7 +157,6 @@ FPrimitiveSceneProxy* UContextualAnimSceneActorComponent::CreateSceneProxy()
 			: FPrimitiveSceneProxy(InComponent)
 			, SceneAssetPtr(InComponent->SceneAsset)
 			, Params(InComponent->DebugParams)
-			, Radius(InComponent->GetScaledSphereRadius())
 		{
 		}
 
@@ -98,7 +201,7 @@ FPrimitiveSceneProxy* UContextualAnimSceneActorComponent::CreateSceneProxy()
 
 		virtual void GetDynamicMeshElements(const TArray<const FSceneView*>& Views, const FSceneViewFamily& ViewFamily, uint32 VisibilityMap, FMeshElementCollector& Collector) const override
 		{
-			const UContextualAnimCompositeSceneAsset* Asset = SceneAssetPtr.Get();
+			const UContextualAnimSceneAssetBase* Asset = SceneAssetPtr.Get();
 			if (Asset == nullptr)
 			{
 				return;
@@ -122,93 +225,102 @@ FPrimitiveSceneProxy* UContextualAnimSceneActorComponent::CreateSceneProxy()
 
 					FPrimitiveDrawInterface* PDI = Collector.GetPDI(ViewIndex);
 
-					// Draw collision sphere
-					const FLinearColor DrawSphereColor = GetViewSelectionColor(FColor::Red, *View, IsSelected(), IsHovered(), false, IsIndividuallySelected());
-
-					float AbsScaleX = LocalToWorld.GetScaledAxis(EAxis::X).Size();
-					float AbsScaleY = LocalToWorld.GetScaledAxis(EAxis::Y).Size();
-					float AbsScaleZ = LocalToWorld.GetScaledAxis(EAxis::Z).Size();
-					float MinAbsScale = FMath::Min3(AbsScaleX, AbsScaleY, AbsScaleZ);
-
-					FVector ScaledX = LocalToWorld.GetUnitAxis(EAxis::X) * MinAbsScale;
-					FVector ScaledY = LocalToWorld.GetUnitAxis(EAxis::Y) * MinAbsScale;
-					FVector ScaledZ = LocalToWorld.GetUnitAxis(EAxis::Z) * MinAbsScale;
+					//DrawCircle(PDI, ToWorldTransform.GetLocation(), FVector(1, 0, 0), FVector(0, 1, 0), FColor::Red, SceneAssetPtr->GetRadius(), 12, SDPG_World, 1.f);
 
-					const int32 SphereSides = FMath::Clamp<int32>(Radius / 4.f, 16, 64);
-					DrawCircle(PDI, LocalToWorld.GetOrigin(), ScaledX, ScaledY, DrawSphereColor, Radius, SphereSides, SDPG_World);
-					DrawCircle(PDI, LocalToWorld.GetOrigin(), ScaledX, ScaledZ, DrawSphereColor, Radius, SphereSides, SDPG_World);
-					DrawCircle(PDI, LocalToWorld.GetOrigin(), ScaledY, ScaledZ, DrawSphereColor, Radius, SphereSides, SDPG_World);
-
-					FContextualAnimQueryResult Result;
-					if (Params.TestActor.IsValid())
+					SceneAssetPtr->ForEachAnimData([=](const FName& Role, const FContextualAnimData& Data)
 					{
-						if (Asset->QueryData(Result, FContextualAnimQueryParams(Params.TestActor.Get(), true, true), ToWorldTransform))
-						{							
-							const float Time = Result.AnimStartTime;
-							const FTransform TransformAtTime = Asset->InteractorTrack.AnimDataContainer[Result.DataIndex].GetAlignmentTransformAtTime(Time) * ToWorldTransform;
-							DrawCoordinateSystem(PDI, TransformAtTime.GetLocation(), TransformAtTime.Rotator(), 20.f, SDPG_World, 2.f);
-						}
-					}
-
-					for (int32 Idx = 0; Idx < Asset->InteractorTrack.AnimDataContainer.Num(); Idx++)
-					{
-						const FContextualAnimData& Data = Asset->InteractorTrack.AnimDataContainer[Idx];
-
-						FLinearColor DrawColor = Result.DataIndex == Idx ? FLinearColor::Red : FLinearColor::White;
-
-						// Draw Entry Point
-						const FTransform EntryTransform = (Data.GetAlignmentTransformAtEntryTime() * ToWorldTransform);
-						DrawCoordinateSystem(PDI, EntryTransform.GetLocation(), EntryTransform.Rotator(), 20.f, SDPG_World, 3.f);
-
-						// Draw Facing Tolerance
-						if(Data.Metadata && Data.Metadata->Facing.Tolerance > 0.f)
+						if (Role != SceneAssetPtr->PrimaryRole)
 						{
-							DrawSector(PDI, EntryTransform.GetLocation(), EntryTransform.GetRotation().GetForwardVector(), 0.f, 30.f, -Data.Metadata->Facing.Tolerance, Data.Metadata->Facing.Tolerance, DrawColor, SDPG_World, 1.f);
-						}
-						else
-						{
-							DrawCircle(PDI, EntryTransform.GetLocation(), FVector(1, 0, 0), FVector(0, 1, 0), DrawColor, 30.f, 12, SDPG_World, 1.f);
-						}
+							FLinearColor DrawColor = FLinearColor::White;
 
-						// Draw Sync Point
-						const FTransform SyncPoint = Data.GetAlignmentTransformAtSyncTime() * ToWorldTransform;
-						DrawCoordinateSystem(PDI, SyncPoint.GetLocation(), SyncPoint.Rotator(), 20.f, SDPG_World, 3.f);
+							//@TODO: Temp, this should not be here
+							if(Params.TestActor.IsValid())
+							{
+								if (Data.Metadata)
+								{
+									const FTransform EntryTransform = Data.GetAlignmentTransformAtEntryTime() * ToWorldTransform;
+									if (Data.Metadata->DoesQuerierPassConditions(FContextualAnimQuerier(Params.TestActor.Get()), FContextualAnimQueryContext(ToWorldTransform), EntryTransform))
+									{
+										DrawColor = FLinearColor::Red;
+									}
+								}
+							}
 
-						if(Params.DrawAlignmentTransformAtTime != 0.f)
-						{
-							const FTransform RootAtTime = (Data.GetAlignmentTransformAtTime(Params.DrawAlignmentTransformAtTime) * ToWorldTransform);
-							DrawCoordinateSystem(PDI, RootAtTime.GetLocation(), RootAtTime.Rotator(), 10.f, SDPG_World, 2.f);
-						}
+							// Draw Entry Point
+							const FTransform EntryTransform = (Data.GetAlignmentTransformAtEntryTime() * ToWorldTransform);
+							DrawCoordinateSystem(PDI, EntryTransform.GetLocation(), EntryTransform.Rotator(), 20.f, SDPG_World, 3.f);
 
-						// Draw Sector
-						if(Data.Metadata)
-						{
-							FVector Origin = ToWorldTransform.GetLocation();
-							FVector Direction = (EntryTransform.GetLocation() - ToWorldTransform.GetLocation()).GetSafeNormal2D();
+							// Draw Sync Point
+							const FTransform SyncPoint = Data.GetAlignmentTransformAtSyncTime() * ToWorldTransform;
+							DrawCoordinateSystem(PDI, SyncPoint.GetLocation(), SyncPoint.Rotator(), 20.f, SDPG_World, 3.f);
 
-							if (Data.Metadata->OffsetFromOrigin != 0.f)
+							if (Params.DrawAlignmentTransformAtTime != 0.f)
 							{
-								Origin = Origin + Direction * Data.Metadata->OffsetFromOrigin;
+								const FTransform RootAtTime = (Data.GetAlignmentTransformAtTime(Params.DrawAlignmentTransformAtTime) * ToWorldTransform);
+								DrawCoordinateSystem(PDI, RootAtTime.GetLocation(), RootAtTime.Rotator(), 10.f, SDPG_World, 2.f);
 							}
 
-							if (Data.Metadata->Angle.Tolerance > 0.f)
+							// Draw Facing Tolerance
+							if (Data.Metadata && Data.Metadata->Facing > 0.f)
 							{
-								DrawSector(PDI, Origin, Direction, Data.Metadata->Distance.MinDistance, Data.Metadata->Distance.MaxDistance, -Data.Metadata->Angle.Tolerance, Data.Metadata->Angle.Tolerance, DrawColor, SDPG_World, 3.f);
+								DrawSector(PDI, EntryTransform.GetLocation(), EntryTransform.GetRotation().GetForwardVector(), 0.f, 30.f, -Data.Metadata->Facing, Data.Metadata->Facing, DrawColor, SDPG_World, 1.f);
 							}
 							else
 							{
-								if (Data.Metadata->Distance.MinDistance > 0.f)
+								DrawCircle(PDI, EntryTransform.GetLocation(), FVector(1, 0, 0), FVector(0, 1, 0), DrawColor, 30.f, 12, SDPG_World, 1.f);
+							}
+
+							// Draw Sector
+							if (Data.Metadata)
+							{
+								FVector Origin = ToWorldTransform.GetLocation();
+								Origin.Z = EntryTransform.GetLocation().Z;
+
+								FVector Direction = (EntryTransform.GetLocation() - ToWorldTransform.GetLocation()).GetSafeNormal2D();
+
+								if (Data.Metadata->DirectionOffset != 0.f)
+								{
+									Direction = Direction.RotateAngleAxis(Data.Metadata->DirectionOffset, FVector::UpVector);
+								}
+
+								if (Data.Metadata->OriginOffset.X != 0.f)
 								{
-									DrawCircle(PDI, Origin, FVector(1, 0, 0), FVector(0, 1, 0), DrawColor, Data.Metadata->Distance.MinDistance, 12, SDPG_World, 2.f);
+									Origin = Origin + Direction * Data.Metadata->OriginOffset.X;
 								}
 
-								if (Data.Metadata->Distance.MaxDistance > 0.f)
+								if (Data.Metadata->OriginOffset.Y != 0.f)
 								{
-									DrawCircle(PDI, Origin, FVector(1, 0, 0), FVector(0, 1, 0), DrawColor, Data.Metadata->Distance.MaxDistance, 12, SDPG_World, 2.f);
+									Origin = Origin + (Direction.ToOrientationQuat().GetRightVector()) * Data.Metadata->OriginOffset.Y;
+								}
+
+								if (Data.Metadata->NearWidth > 0.f || Data.Metadata->FarWidth > 0.f)
+								{
+									const float HalfNearWidth = (Data.Metadata->NearWidth / 2.f);
+									const FVector RightVector = Direction.ToOrientationQuat().GetRightVector();
+									const FVector A = Origin + (-RightVector * HalfNearWidth);
+									const FVector B = Origin + (RightVector * HalfNearWidth);
+
+									const float HalfFarWidth = (Data.Metadata->FarWidth / 2.f);
+									const FVector FarEdgeCenter = Origin + (Direction * Data.Metadata->MaxDistance);
+									const FVector C = FarEdgeCenter + (-RightVector * HalfFarWidth);
+									const FVector D = FarEdgeCenter + (RightVector * HalfFarWidth);
+
+									PDI->DrawLine(A, B, DrawColor, SDPG_World, 2.f);
+									PDI->DrawLine(C, D, DrawColor, SDPG_World, 2.f);
+
+									PDI->DrawLine(A, C, DrawColor, SDPG_World, 2.f);
+									PDI->DrawLine(B, D, DrawColor, SDPG_World, 2.f);
+								}
+								else
+								{
+									PDI->DrawLine(Origin, Origin + Direction * Data.Metadata->MaxDistance, DrawColor, SDPG_World, 2.f);
+									DrawCircle(PDI, Origin, FVector(1, 0, 0), FVector(0, 1, 0), DrawColor, Data.Metadata->MaxDistance, 32, SDPG_World, 2.f);
 								}
 							}
 						}
-					}
+
+						return EContextualAnimForEachResult::Continue;
+					});
 				}
 			}
 		}
@@ -234,9 +346,8 @@ FPrimitiveSceneProxy* UContextualAnimSceneActorComponent::CreateSceneProxy()
 		}
 
 	private:
-		TWeakObjectPtr<const UContextualAnimCompositeSceneAsset> SceneAssetPtr;
+		TWeakObjectPtr<const UContextualAnimSceneAssetBase> SceneAssetPtr;
 		FContextualAnimDebugParams Params;
-		float Radius;
 	};
 
 	if(bEnableDebug)
diff --git a/Engine/Plugins/Experimental/Animation/ContextualAnimation/Source/ContextualAnimation/Private/ContextualAnimSceneAsset.cpp b/Engine/Plugins/Experimental/Animation/ContextualAnimation/Source/ContextualAnimation/Private/ContextualAnimSceneAsset.cpp
--- a/Engine/Plugins/Experimental/Animation/ContextualAnimation/Source/ContextualAnimation/Private/ContextualAnimSceneAsset.cpp
+++ b/Engine/Plugins/Experimental/Animation/ContextualAnimation/Source/ContextualAnimation/Private/ContextualAnimSceneAsset.cpp
@@ -1,26 +1,20 @@
 // Copyright Epic Games, Inc. All Rights Reserved.
 
 #include "ContextualAnimSceneAsset.h"
-
 #include "AnimationRuntime.h"
 #include "Animation/AnimMontage.h"
 #include "Animation/AnimationPoseData.h"
 #include "Animation/AnimTypes.h"
 #include "ContextualAnimUtilities.h"
+#include "ContextualAnimMetadata.h"
 #include "UObject/ObjectSaveContext.h"
+#include "ContextualAnimScenePivotProvider.h"
 
 UContextualAnimSceneAsset::UContextualAnimSceneAsset(const FObjectInitializer& ObjectInitializer)
 	: Super(ObjectInitializer)
 {
 }
 
-void UContextualAnimSceneAsset::PreSave(const class ITargetPlatform* TargetPlatform)
-{
-	PRAGMA_DISABLE_DEPRECATION_WARNINGS;
-	Super::PreSave(TargetPlatform);
-	PRAGMA_ENABLE_DEPRECATION_WARNINGS;
-}
-
 void UContextualAnimSceneAsset::PreSave(FObjectPreSaveContext ObjectSaveContext)
 {
 	// Necessary for FCompactPose that uses a FAnimStackAllocator (TMemStackAllocator) which allocates from FMemStack.
@@ -29,95 +23,97 @@ void UContextualAnimSceneAsset::PreSave(FObjectPreSaveContext ObjectSaveContext)
 	// but any allocation from outside the game tick (like here when generating the alignment tracks off-line) must explicitly add a mark to avoid a leak 
 	FMemMark Mark(FMemStack::Get());
 
-	// Generate scene pivot for each alignment section
 	Super::PreSave(ObjectSaveContext);
-	
-	const FTransform MeshToComponentInverse = MeshToComponent.Inverse();
-	const float SampleInterval = 1.f / SampleRate;
-
-	TArray<FBoneIndexType> RequiredBoneIndexArray = { 0 };
 
+	// Set Index for each ContextualAnimData
+	int32 NumAnimData = 0;
 	for (auto& Entry : DataContainer)
 	{
-		// Generate alignment tracks relative to scene pivot
-		GenerateAlignmentTracksRelativeToScenePivot(Entry.Value.AnimData);
-
-		// Generate alignment tracks relative to other actors
-		const FName& Role = Entry.Key;
-		FContextualAnimData& AnimData = Entry.Value.AnimData;
-
-		const UAnimMontage* Animation = AnimData.Animation;
+		NumAnimData = FMath::Max(NumAnimData, Entry.Value.AnimDataContainer.Num());
+		for (int32 AnimDataIdx = 0; AnimDataIdx < Entry.Value.AnimDataContainer.Num(); AnimDataIdx++)
+		{
+			Entry.Value.AnimDataContainer[AnimDataIdx].Index = AnimDataIdx;
+		}
+	}
 
-		UE_LOG(LogContextualAnim, Log, TEXT("%s Generating Alignment Track Relative To Others for '%s' (Anim: %s)"), *GetNameSafe(this), *Role.ToString(), *GetNameSafe(Animation));
+	// Generate scene pivot for each alignment section
+	for (int32 AlignmentSectionIdx = 0; AlignmentSectionIdx < AlignmentSections.Num(); AlignmentSectionIdx++)
+	{
+		AlignmentSections[AlignmentSectionIdx].ScenePivots.Reset();
 
-		if (Animation)
+		// We need to calculate scene pivot for each set of animations
+		for (int32 AnimDataIdx = 0; AnimDataIdx < NumAnimData; AnimDataIdx++)
 		{
-			FBoneContainer BoneContainer = FBoneContainer(RequiredBoneIndexArray, FCurveEvaluationOption(false), *Animation->GetSkeleton());
-
-			for (const auto& OtherEntry : DataContainer)
+			if (AlignmentSections[AlignmentSectionIdx].ScenePivotProvider)
 			{
-				const FName& OtherRole = OtherEntry.Key;
-				const FContextualAnimData& OtherAnimData = OtherEntry.Value.AnimData;
-
-				if (Role != OtherRole)
-				{
-					float Time = 0.f;
-					float EndTime = Animation->GetPlayLength();
-					int32 SampleIndex = 0;
-
-					AnimData.AlignmentData.Tracks.TrackNames.Add(OtherRole);
-					FRawAnimSequenceTrack& NewTrack = AnimData.AlignmentData.Tracks.AnimationTracks.AddZeroed_GetRef();
-
-					const UAnimMontage* OtherAnimation = OtherAnimData.Animation;
-					if (OtherAnimation)
-					{
-						FBoneContainer OtherBoneContainer = FBoneContainer(RequiredBoneIndexArray, FCurveEvaluationOption(false), *OtherAnimation->GetSkeleton());
-
-						while (Time < EndTime)
-						{
-							Time = FMath::Clamp(SampleIndex * SampleInterval, 0.f, EndTime);
-							SampleIndex++;
-
-							FCSPose<FCompactPose> ComponentSpacePose;
-							UContextualAnimUtilities::ExtractComponentSpacePose(Animation, BoneContainer, Time, false, ComponentSpacePose);
-
-							FCSPose<FCompactPose> OtherComponentSpacePose;
-							UContextualAnimUtilities::ExtractComponentSpacePose(OtherAnimation, OtherBoneContainer, Time, false, OtherComponentSpacePose);
-
-							const FTransform RootTransform = MeshToComponentInverse * (ComponentSpacePose.GetComponentSpaceTransform(FCompactPoseBoneIndex(0)) * AnimData.MeshToScene);
-							const FTransform OtherRootTransform = MeshToComponentInverse * (OtherComponentSpacePose.GetComponentSpaceTransform(FCompactPoseBoneIndex(0)) * OtherAnimData.MeshToScene);
-							const FTransform RootRelativeToRef = RootTransform.GetRelativeTransform(OtherRootTransform);
-
-							NewTrack.PosKeys.Add(RootRelativeToRef.GetLocation());
-							NewTrack.RotKeys.Add(RootRelativeToRef.GetRotation());
-							NewTrack.ScaleKeys.Add(RootRelativeToRef.GetScale3D());
-						};
-					}
-
-					UE_LOG(LogContextualAnim, Log, TEXT("\t Relative To: '%s' (Anim: %s) NumKeys: %d"), *OtherRole.ToString(), *GetNameSafe(OtherAnimation), NewTrack.PosKeys.Num());
-				}
+				const FTransform ScenePivot = AlignmentSections[AlignmentSectionIdx].ScenePivotProvider->CalculateScenePivot_Source(AnimDataIdx);
+				AlignmentSections[AlignmentSectionIdx].ScenePivots.Add(ScenePivot);
+			}
+			else
+			{
+				AlignmentSections[AlignmentSectionIdx].ScenePivots.Add(FTransform::Identity);
 			}
 		}
 	}
+	
+
+	for (auto& Entry : DataContainer)
+	{
+		for (FContextualAnimData& Data : Entry.Value.AnimDataContainer)
+		{
+			// Generate alignment tracks relative to scene pivot
+			GenerateAlignmentTracks(Entry.Value.Settings, Data);
+
+			// Generate IK Targets
+			GenerateIKTargetTracks(Entry.Value.Settings, Data);
+		}
+	}
+
+	UpdateRadius();
 }
 
-const FContextualAnimTrack* UContextualAnimSceneAsset::FindTrack(const FName& Role) const
+const FContextualAnimTrackSettings* UContextualAnimSceneAsset::GetTrackSettings(const FName& Role) const
 {
-	return DataContainer.Find(Role);
+	const FContextualAnimCompositeTrack* Track = DataContainer.Find(Role);
+	return Track ? &Track->Settings : nullptr;
 }
 
-UClass* UContextualAnimSceneAsset::GetPreviewActorClassForRole(const FName& Role) const
+const FContextualAnimData* UContextualAnimSceneAsset::GetAnimDataForRoleAtIndex(const FName& Role, int32 Index) const
 {
-	if (const FContextualAnimTrack* Result = FindTrack(Role))
+	if(const FContextualAnimCompositeTrack* Track = DataContainer.Find(Role))
 	{
-		return Result->Settings.PreviewActorClass;
+		if(Track->AnimDataContainer.IsValidIndex(Index))
+		{
+			return &Track->AnimDataContainer[Index];
+		}
 	}
 
 	return nullptr;
 }
 
-EContextualAnimJoinRule UContextualAnimSceneAsset::GetJoinRuleForRole(const FName& Role) const
+void UContextualAnimSceneAsset::ForEachAnimData(FForEachAnimDataFunction Function) const
 {
-	const FContextualAnimTrack* Result = FindTrack(Role);
-	return Result ? Result->Settings.JoinRule : EContextualAnimJoinRule::Default;
+	for (const auto& Pair : DataContainer)
+	{
+		for (const FContextualAnimData& Data : Pair.Value.AnimDataContainer)
+		{
+			if (Function(Pair.Key, Data) == EContextualAnimForEachResult::Break)
+			{
+				return;
+			}
+		}
+	}
 }
+
+TArray<FName> UContextualAnimSceneAsset::GetRoles() const
+{
+	TArray<FName> Roles;
+	DataContainer.GetKeys(Roles);
+	return Roles;
+}
+
+bool UContextualAnimSceneAsset::Query(const FName& Role, FContextualAnimQueryResult& OutResult, const FContextualAnimQueryParams& QueryParams, const FTransform& ToWorldTransform) const
+{
+	const FContextualAnimCompositeTrack* Track = DataContainer.Find(Role);
+	return Track ? QueryCompositeTrack(Track, OutResult, QueryParams, ToWorldTransform) : false;
+}
\ No newline at end of file
diff --git a/Engine/Plugins/Experimental/Animation/ContextualAnimation/Source/ContextualAnimation/Private/ContextualAnimSceneAssetBase.cpp b/Engine/Plugins/Experimental/Animation/ContextualAnimation/Source/ContextualAnimation/Private/ContextualAnimSceneAssetBase.cpp
--- a/Engine/Plugins/Experimental/Animation/ContextualAnimation/Source/ContextualAnimation/Private/ContextualAnimSceneAssetBase.cpp
+++ b/Engine/Plugins/Experimental/Animation/ContextualAnimation/Source/ContextualAnimation/Private/ContextualAnimSceneAssetBase.cpp
@@ -12,111 +12,344 @@
 #include "ContextualAnimMetadata.h"
 #include "UObject/ObjectSaveContext.h"
 
+static FCompactPoseBoneIndex GetCompactPoseBoneIndexFromPose(const FCSPose<FCompactPose>& Pose, const FName& BoneName)
+{
+	const FBoneContainer& BoneContainer = Pose.GetPose().GetBoneContainer();
+	for (int32 Idx = Pose.GetPose().GetNumBones() - 1; Idx >= 0; Idx--)
+	{
+		if (BoneContainer.GetReferenceSkeleton().GetBoneName(BoneContainer.GetBoneIndicesArray()[Idx]) == BoneName)
+		{
+			return FCompactPoseBoneIndex(Idx);
+		}
+	}
+
+	checkf(false, TEXT("BoneName: %s Pose.Asset: %s Pose.NumBones: %d"), *BoneName.ToString(), *GetNameSafe(Pose.GetPose().GetBoneContainer().GetAsset()), Pose.GetPose().GetNumBones());
+	return FCompactPoseBoneIndex(INDEX_NONE);
+}
+
 UContextualAnimSceneAssetBase::UContextualAnimSceneAssetBase(const FObjectInitializer& ObjectInitializer)
 	: Super(ObjectInitializer)
 {
+	bDisableCollisionBetweenActors = true;
 	SampleRate = 15;
-	MeshToComponent = FTransform(FRotator(0.f, -90.f, 0.f));
 }
 
-FTransform UContextualAnimSceneAssetBase::ExtractTransformFromAnimData(const FContextualAnimData& AnimData, float Time) const
+void UContextualAnimSceneAssetBase::GenerateAlignmentTracks(const FContextualAnimTrackSettings& Settings, FContextualAnimData& AnimData)
 {
-	const FTransform MeshToComponentInverse = MeshToComponent.Inverse();
+	const FTransform MeshToComponentInverse = Settings.MeshToComponent.Inverse();
+	const float SampleInterval = 1.f / SampleRate;
+
+	// Initialize tracks for each alignment section
+	const int32 TotalTracks = AlignmentSections.Num();
+	AnimData.AlignmentData.Initialize(TotalTracks, SampleInterval);
+	for (int32 Idx = 0; Idx < AlignmentSections.Num(); Idx++)
+	{
+		AnimData.AlignmentData.Tracks.TrackNames.Add(AlignmentSections[Idx].SectionName);
+		AnimData.AlignmentData.Tracks.AnimationTracks.AddZeroed();
+	}
 
 	if (const UAnimMontage* Animation = AnimData.Animation)
 	{
-		FCSPose<FCompactPose> ComponentSpacePose;
-		UContextualAnimUtilities::ExtractComponentSpacePose(Animation, FBoneContainer({ 0 }, FCurveEvaluationOption(false), *Animation->GetSkeleton()), Time, false, ComponentSpacePose);
-		return MeshToComponentInverse * (ComponentSpacePose.GetComponentSpaceTransform(FCompactPoseBoneIndex(0)) * AnimData.MeshToScene);
+		float Time = 0.f;
+		float EndTime = Animation->GetPlayLength();
+		int32 SampleIndex = 0;
+		while (Time < EndTime)
+		{
+			Time = FMath::Clamp(SampleIndex * SampleInterval, 0.f, EndTime);
+			SampleIndex++;
+			
+			const FTransform RootTransform = MeshToComponentInverse * (UContextualAnimUtilities::ExtractRootTransformFromAnimation(Animation, Time) * AnimData.MeshToScene);
+
+			for (int32 Idx = 0; Idx < AlignmentSections.Num(); Idx++)
+			{
+				FRawAnimSequenceTrack& SceneTrack = AnimData.AlignmentData.Tracks.AnimationTracks[Idx];
+
+				const FTransform ScenePivotTransform = AlignmentSections[Idx].ScenePivots[AnimData.Index];
+				const FTransform RootRelativeToScenePivot = RootTransform.GetRelativeTransform(ScenePivotTransform);
+
+				SceneTrack.PosKeys.Add(RootRelativeToScenePivot.GetLocation());
+				SceneTrack.RotKeys.Add(RootRelativeToScenePivot.GetRotation());
+			}
+		}
 	}
 	else
 	{
-		return MeshToComponentInverse * AnimData.MeshToScene;
+		for (int32 Idx = 0; Idx < AlignmentSections.Num(); Idx++)
+		{
+			AnimData.AlignmentData.Tracks.TrackNames.Add(AlignmentSections[Idx].SectionName);
+			FRawAnimSequenceTrack& SceneTrack = AnimData.AlignmentData.Tracks.AnimationTracks.AddZeroed_GetRef();
+
+			const FTransform RootTransform = MeshToComponentInverse * AnimData.MeshToScene;
+			const FTransform ScenePivotTransform = AlignmentSections[Idx].ScenePivots[AnimData.Index];
+			const FTransform RootRelativeToScenePivot = RootTransform.GetRelativeTransform(ScenePivotTransform);
+
+			SceneTrack.PosKeys.Add(RootRelativeToScenePivot.GetLocation());
+			SceneTrack.RotKeys.Add(RootRelativeToScenePivot.GetRotation());
+		}
 	}
 }
 
-void UContextualAnimSceneAssetBase::GenerateAlignmentTracksRelativeToScenePivot(FContextualAnimData& AnimData)
+void UContextualAnimSceneAssetBase::GenerateIKTargetTracks(const FContextualAnimTrackSettings& Settings, FContextualAnimData& AnimData)
 {
-	const FTransform MeshToComponentInverse = MeshToComponent.Inverse();
-	const float SampleInterval = 1.f / SampleRate;
+	AnimData.IKTargetData.Empty();
 
-	const int32 TotalTracks = AlignmentSections.Num();
-	AnimData.AlignmentData.Tracks.AnimationTracks.Reset(TotalTracks);
-	AnimData.AlignmentData.Tracks.TrackNames.Reset(TotalTracks);
-	AnimData.AlignmentData.SampleInterval = SampleInterval;
+	if(Settings.IKTargetDefinitions.Num() == 0)
+	{
+		return;
+	}
 
-	const UAnimMontage* Animation = AnimData.Animation;
-	if (Animation)
+	if (const UAnimMontage* Animation = AnimData.Animation)
 	{
+		UE_LOG(LogContextualAnim, Log, TEXT("%s Generating IK Target Tracks. Animation: %s"), *GetNameSafe(this), *GetNameSafe(Animation));
+
+		const float SampleInterval = 1.f / SampleRate;
+
 		TArray<FBoneIndexType> RequiredBoneIndexArray;
-		RequiredBoneIndexArray.Add(0);
-		FBoneContainer BoneContainer = FBoneContainer(RequiredBoneIndexArray, FCurveEvaluationOption(false), *Animation->GetSkeleton());
 
-		for (int32 Idx = 0; Idx < AlignmentSections.Num(); Idx++)
+		// Helper structure to group pose extraction per target so we can extract the pose for all the bones that are relative to the same target in one pass
+		struct FPoseExtractionHelper
+		{
+			const FContextualAnimData* TargetAnimDataPtr = nullptr;
+			TArray<TTuple<FName, int32, FName, int32>> BonesData; //0: MyBoneName 1: MyBoneIndex, 2: TargetBoneName, 3 TargetBoneIndex
+		};
+		TMap<FName, FPoseExtractionHelper> PoseExtractionHelperMap;
+		PoseExtractionHelperMap.Reserve(Settings.IKTargetDefinitions.Num());
+
+		int32 TotalTracks = 0;
+		for (const FContextualAnimIKTargetDefinition& IKTargetDef : Settings.IKTargetDefinitions)
+		{
+			if (IKTargetDef.Provider == EContextualAnimIKTargetProvider::Autogenerated)
+			{
+				const FName TargetRole = IKTargetDef.AutoParams.TargetRole;
+				FPoseExtractionHelper* DataPtr = PoseExtractionHelperMap.Find(TargetRole);
+				if(DataPtr == nullptr)
+				{
+					// Find AnimData for target role. 
+					const FContextualAnimData* TargetAnimDataPtr = GetAnimDataForRoleAtIndex(TargetRole, AnimData.Index);
+					if (TargetAnimDataPtr == nullptr)
+					{
+						UE_LOG(LogContextualAnim, Warning, TEXT("\t Can't find AnimTrack for TargetRole '%s'"), *TargetRole.ToString());
+						continue;
+					}
+
+					DataPtr = &PoseExtractionHelperMap.Add(TargetRole);
+					DataPtr->TargetAnimDataPtr = TargetAnimDataPtr;
+				}
+				
+				//@TODO: Replace with the actual bone. We need a way to go from IKGoalName to BoneName
+				const FName BoneName = IKTargetDef.IKGoalName;
+				const int32 BoneIndex = Animation->GetSkeleton()->GetReferenceSkeleton().FindBoneIndex(BoneName);
+				if (BoneIndex == INDEX_NONE)
+				{
+					UE_LOG(LogContextualAnim, Warning, TEXT("\t Can't find BoneIndex. BoneName: %s Animation: %s Skel: %s"), 
+						*BoneName.ToString(), *GetNameSafe(Animation), *GetNameSafe(Animation->GetSkeleton()));
+
+					continue;
+				}
+
+				// Find TargetBoneIndex. Note that we add TargetBoneIndexeven if it is INDEX_NONE. In this case, my bone will be relative to the origin of the target actor. 
+				// This is to support cases where the target actor doesn't have animation or TargetBoneName is None
+				FName TargetBoneName = IKTargetDef.AutoParams.BoneName;
+				const UAnimMontage* TargetAnimation = DataPtr->TargetAnimDataPtr->Animation;
+				const int32 TargetBoneIndex = TargetAnimation ? TargetAnimation->GetSkeleton()->GetReferenceSkeleton().FindBoneIndex(TargetBoneName) : INDEX_NONE;
+				if (TargetBoneIndex == INDEX_NONE)
+				{
+					UE_LOG(LogContextualAnim, Log, TEXT("\t Can't find TargetBoneIndex. BoneName: %s Animation: %s Skel: %s. Track for this bone will be relative to the origin of the target role."),
+						*TargetBoneName.ToString(), *GetNameSafe(TargetAnimation), TargetAnimation ? *GetNameSafe(TargetAnimation->GetSkeleton()) : nullptr);
+
+					TargetBoneName = NAME_None;
+				}
+
+				RequiredBoneIndexArray.AddUnique(BoneIndex);
+
+				DataPtr->BonesData.Add(MakeTuple(BoneName, BoneIndex, TargetBoneName, TargetBoneIndex));
+				TotalTracks++;
+
+				UE_LOG(LogContextualAnim, Log, TEXT("\t Bone added for extraction. BoneName: %s (%d) TargetRole: %s TargetAnimation: %s TargetBone: %s (%d)"),
+					*BoneName.ToString(), BoneIndex, *TargetRole.ToString(), *GetNameSafe(TargetAnimation), *TargetBoneName.ToString(), TargetBoneIndex);
+			}
+		}
+
+		if(TotalTracks > 0)
 		{
+			// Complete bones chain and create bone container to extract pose from my animation
+			Animation->GetSkeleton()->GetReferenceSkeleton().EnsureParentsExistAndSort(RequiredBoneIndexArray);
+			FBoneContainer BoneContainer = FBoneContainer(RequiredBoneIndexArray, FCurveEvaluationOption(false), *Animation->GetSkeleton());
+
+			// Initialize track container
+			AnimData.IKTargetData.Initialize(TotalTracks, SampleInterval);
+
+			// Initialize lookup map to go from track name to target role and bone (the bone this track is relative to)
+			AnimData.IKTargetTrackLookupMap.Empty(TotalTracks);
+
 			float Time = 0.f;
 			float EndTime = Animation->GetPlayLength();
 			int32 SampleIndex = 0;
-
-			AnimData.AlignmentData.Tracks.TrackNames.Add(AlignmentSections[Idx].SectionName);
-			FRawAnimSequenceTrack& SceneTrack = AnimData.AlignmentData.Tracks.AnimationTracks.AddZeroed_GetRef();
 			while (Time < EndTime)
 			{
 				Time = FMath::Clamp(SampleIndex * SampleInterval, 0.f, EndTime);
 				SampleIndex++;
 
+				// Extract pose from my animation
 				FCSPose<FCompactPose> ComponentSpacePose;
 				UContextualAnimUtilities::ExtractComponentSpacePose(Animation, BoneContainer, Time, false, ComponentSpacePose);
 
-				const FTransform RootTransform = MeshToComponentInverse * (ComponentSpacePose.GetComponentSpaceTransform(FCompactPoseBoneIndex(0)) * AnimData.MeshToScene);
-				const FTransform ScenePivotTransform = AlignmentSections[Idx].ScenePivot;
-				const FTransform RootRelativeToScenePivot = RootTransform.GetRelativeTransform(ScenePivotTransform);
+				// For each target role
+				for (auto& Data : PoseExtractionHelperMap)
+				{
+					// Extract pose from target animation if any
+					FCSPose<FCompactPose> OtherComponentSpacePose;
+					TArray<FBoneIndexType> OtherRequiredBoneIndexArray;
+					FBoneContainer OtherBoneContainer;
+					const UAnimMontage* OtherAnimation = Data.Value.TargetAnimDataPtr->Animation;
+					if (OtherAnimation)
+					{
+						// Prepare array with the indices of the bones to extract from target animation
+						OtherRequiredBoneIndexArray.Reserve(Data.Value.BonesData.Num());
+						for (int32 Idx = 0; Idx < Data.Value.BonesData.Num(); Idx++)
+						{
+							const int32 TargetBoneIndex = Data.Value.BonesData[Idx].Get<3>();
+							if (TargetBoneIndex != INDEX_NONE)
+							{
+								OtherRequiredBoneIndexArray.AddUnique(TargetBoneIndex);
+							}
+						}
 
-				SceneTrack.PosKeys.Add(RootRelativeToScenePivot.GetLocation());
-				SceneTrack.RotKeys.Add(RootRelativeToScenePivot.GetRotation());
-				SceneTrack.ScaleKeys.Add(RootRelativeToScenePivot.GetScale3D());
-			}
-		}
-	}
-	else
-	{
-		for (int32 Idx = 0; Idx < AlignmentSections.Num(); Idx++)
-		{
-			AnimData.AlignmentData.Tracks.TrackNames.Add(AlignmentSections[Idx].SectionName);
-			FRawAnimSequenceTrack& SceneTrack = AnimData.AlignmentData.Tracks.AnimationTracks.AddZeroed_GetRef();
+						if(OtherRequiredBoneIndexArray.Num() > 0)
+						{
+							// Complete bones chain and create bone container to extract pose form the target animation
+							OtherAnimation->GetSkeleton()->GetReferenceSkeleton().EnsureParentsExistAndSort(OtherRequiredBoneIndexArray);
+							OtherBoneContainer = FBoneContainer(OtherRequiredBoneIndexArray, FCurveEvaluationOption(false), *OtherAnimation->GetSkeleton());
 
-			const FTransform RootTransform = MeshToComponentInverse * AnimData.MeshToScene;
-			const FTransform ScenePivotTransform = AlignmentSections[Idx].ScenePivot;
-			const FTransform RootRelativeToScenePivot = RootTransform.GetRelativeTransform(ScenePivotTransform);
+							// Extract pose from target animation
+							UContextualAnimUtilities::ExtractComponentSpacePose(OtherAnimation, OtherBoneContainer, Time, false, OtherComponentSpacePose);
+						}
+					}
 
-			SceneTrack.PosKeys.Add(RootRelativeToScenePivot.GetLocation());
-			SceneTrack.RotKeys.Add(RootRelativeToScenePivot.GetRotation());
-			SceneTrack.ScaleKeys.Add(RootRelativeToScenePivot.GetScale3D());
+					for (int32 Idx = 0; Idx < Data.Value.BonesData.Num(); Idx++)
+					{
+						const FName TrackName = Data.Value.BonesData[Idx].Get<0>();
+						AnimData.IKTargetData.Tracks.TrackNames.Add(TrackName);
+						AnimData.IKTargetData.Tracks.AnimationTracks.AddZeroed();
+
+						// Add entry to the lookup table
+						auto& NewItem = AnimData.IKTargetTrackLookupMap.FindOrAdd(TrackName);
+						NewItem.RoleName = Data.Key;
+						NewItem.BoneName = Data.Value.BonesData[Idx].Get<2>();
+
+						// Get bone transform from my animation
+						const FName BoneName = Data.Value.BonesData[Idx].Get<0>();
+						const FCompactPoseBoneIndex BoneIndex = GetCompactPoseBoneIndexFromPose(ComponentSpacePose, BoneName);
+						const FTransform BoneTransform = (ComponentSpacePose.GetComponentSpaceTransform(BoneIndex) * AnimData.MeshToScene);
+
+						// Get bone transform from target animation
+						FTransform OtherBoneTransform = Data.Value.TargetAnimDataPtr->MeshToScene;
+						const FName TargetBoneName = Data.Value.BonesData[Idx].Get<2>();
+						if (TargetBoneName != NAME_None)
+						{
+							const FCompactPoseBoneIndex OtherBoneIndex = GetCompactPoseBoneIndexFromPose(OtherComponentSpacePose, TargetBoneName);
+							OtherBoneTransform = (OtherComponentSpacePose.GetComponentSpaceTransform(OtherBoneIndex) * Data.Value.TargetAnimDataPtr->MeshToScene);
+						}
+
+						// Get transform relative to target
+						const FTransform BoneRelativeToOther = BoneTransform.GetRelativeTransform(OtherBoneTransform);
+
+						// Add transform to the track
+						FRawAnimSequenceTrack& NewTrack = AnimData.IKTargetData.Tracks.AnimationTracks[Idx];
+						NewTrack.PosKeys.Add(BoneRelativeToOther.GetLocation());
+						NewTrack.RotKeys.Add(BoneRelativeToOther.GetRotation());
+
+						UE_LOG(LogContextualAnim, Verbose, TEXT("\t\t Animation: %s Time: %f BoneName: %s (T: %s) Target Animation: %s TargetBoneName: %s (T: %s)"),
+							*GetNameSafe(Animation), Time, *BoneName.ToString(), *BoneTransform.GetLocation().ToString(),
+							*GetNameSafe(OtherAnimation), *TargetBoneName.ToString(), *OtherBoneTransform.GetLocation().ToString());
+					}
+				}
+			}
 		}
 	}
 }
 
-void UContextualAnimSceneAssetBase::PreSave(const class ITargetPlatform* TargetPlatform)
+void UContextualAnimSceneAssetBase::UpdateRadius()
 {
-	PRAGMA_DISABLE_DEPRECATION_WARNINGS;
-	Super::PreSave(TargetPlatform);
-	PRAGMA_ENABLE_DEPRECATION_WARNINGS;
+	Radius = 0.f;
+	ForEachAnimData([this](const FName& Role, const FContextualAnimData& Data)
+		{
+			Radius = FMath::Max(Radius, (Data.GetAlignmentTransformAtEntryTime().GetLocation().Size()));
+			return EContextualAnimForEachResult::Continue;
+		});
 }
 
-void UContextualAnimSceneAssetBase::PreSave(FObjectPreSaveContext ObjectSaveContext)
+bool UContextualAnimSceneAssetBase::QueryCompositeTrack(const FContextualAnimCompositeTrack* Track, FContextualAnimQueryResult& OutResult, const FContextualAnimQueryParams& QueryParams, const FTransform& ToWorldTransform) const
 {
-	Super::PreSave(ObjectSaveContext);
+	OutResult.Reset();
 
-	// Generate scene pivot for each alignment section
-	for (int32 Idx = 0; Idx < AlignmentSections.Num(); Idx++)
+	const FTransform QueryTransform = QueryParams.Querier.IsValid() ? QueryParams.Querier->GetActorTransform() : QueryParams.QueryTransform;
+
+	int32 DataIndex = INDEX_NONE;
+	if (Track)
 	{
-		if (AlignmentSections[Idx].ScenePivotProvider)
+		if (QueryParams.bComplexQuery)
+		{
+			for (int32 Idx = 0; Idx < Track->AnimDataContainer.Num(); Idx++)
+			{
+				const FContextualAnimData& Data = Track->AnimDataContainer[Idx];
+
+				if (Data.Metadata)
+				{
+					const FTransform EntryTransform = Data.GetAlignmentTransformAtEntryTime() * ToWorldTransform;
+
+					if (!Data.Metadata->DoesQuerierPassConditions(FContextualAnimQuerier(QueryTransform), FContextualAnimQueryContext(ToWorldTransform), EntryTransform))
+					{
+						continue;
+					}
+				}
+
+				// Return the first item that passes all tests
+				DataIndex = Idx;
+				break;
+			}
+		}
+		else // Simple Query
 		{
-			AlignmentSections[Idx].ScenePivot = AlignmentSections[Idx].ScenePivotProvider->CalculateScenePivot_Source();
+			float BestDistanceSq = MAX_FLT;
+			for (int32 Idx = 0; Idx < Track->AnimDataContainer.Num(); Idx++)
+			{
+				const FContextualAnimData& Data = Track->AnimDataContainer[Idx];
+
+				//@TODO: Convert querier location to local space instead
+				const FTransform EntryTransform = Data.GetAlignmentTransformAtEntryTime() * ToWorldTransform;
+				const float DistSq = FVector::DistSquared2D(EntryTransform.GetLocation(), QueryTransform.GetLocation());
+				if (DistSq < BestDistanceSq)
+				{
+					BestDistanceSq = DistSq;
+					DataIndex = Idx;
+				}
+			}
 		}
-		else
+
+		if (DataIndex != INDEX_NONE)
 		{
-			AlignmentSections[Idx].ScenePivot = MeshToComponent.Inverse();
+			const FContextualAnimData& ResultData = Track->AnimDataContainer[DataIndex];
+
+			OutResult.DataIndex = DataIndex;
+			OutResult.Animation = ResultData.Animation;
+			OutResult.EntryTransform = ResultData.GetAlignmentTransformAtEntryTime() * ToWorldTransform;
+			OutResult.SyncTransform = ResultData.GetAlignmentTransformAtSyncTime() * ToWorldTransform;
+
+			if (QueryParams.bFindAnimStartTime)
+			{
+				const FVector LocalLocation = (QueryTransform.GetRelativeTransform(ToWorldTransform)).GetLocation();
+				OutResult.AnimStartTime = ResultData.FindBestAnimStartTime(LocalLocation);
+			}
+
+			return true;
 		}
 	}
+
+	return false;
+}
+
+void UContextualAnimSceneAssetBase::PreSave(FObjectPreSaveContext ObjectSaveContext)
+{
+	Super::PreSave(ObjectSaveContext);
 }
diff --git a/Engine/Plugins/Experimental/Animation/ContextualAnimation/Source/ContextualAnimation/Private/ContextualAnimSceneInstance.cpp b/Engine/Plugins/Experimental/Animation/ContextualAnimation/Source/ContextualAnimation/Private/ContextualAnimSceneInstance.cpp
--- a/Engine/Plugins/Experimental/Animation/ContextualAnimation/Source/ContextualAnimation/Private/ContextualAnimSceneInstance.cpp
+++ b/Engine/Plugins/Experimental/Animation/ContextualAnimation/Source/ContextualAnimation/Private/ContextualAnimSceneInstance.cpp
@@ -11,70 +11,84 @@
 #include "ContextualAnimScenePivotProvider.h"
 #include "ContextualAnimSceneAsset.h"
 #include "ContextualAnimSceneActorComponent.h"
+#include "ContextualAnimTransition.h"
+#include "ContextualAnimUtilities.h"
 #include "Engine/World.h"
 #include "DrawDebugHelpers.h"
 
-//@TODO: Move to ContextualAnimUtilities
-static UAnimInstance* TryGetAnimInstance(AActor* Actor)
+UContextualAnimSceneActorComponent* FContextualAnimSceneActorData::GetSceneActorComponent() const
 {
-	UAnimInstance* AnimInstance = nullptr;
-	if (Actor)
-	{
-		if (ACharacter* Character = Cast<ACharacter>(Actor))
-		{
-			USkeletalMeshComponent* SkelMeshComp = Character->GetMesh();
-			AnimInstance = SkelMeshComp ? SkelMeshComp->GetAnimInstance() : nullptr;
-		}
-		else if (Actor->GetClass()->ImplementsInterface(UContextualAnimActorInterface::StaticClass()))
-		{
-			USkeletalMeshComponent* SkelMeshComp = IContextualAnimActorInterface::Execute_GetMesh(Actor);
-			AnimInstance = SkelMeshComp ? SkelMeshComp->GetAnimInstance() : nullptr;
-		}
-		else if (USkeletalMeshComponent* SkelMeshComp = Actor->FindComponentByClass<USkeletalMeshComponent>())
-		{
-			AnimInstance = SkelMeshComp->GetAnimInstance();
-		}
-	}
+	//@TODO: Cache this during the binding
+	return (Actor.IsValid()) ? Actor->FindComponentByClass<UContextualAnimSceneActorComponent>() : nullptr;
+}
 
-	return AnimInstance;
+FTransform FContextualAnimSceneActorData::GetTransform() const
+{
+	const UContextualAnimSceneActorComponent* Comp = GetSceneActorComponent();
+	return Comp ? Comp->GetComponentTransform() : Actor->GetActorTransform();
 }
 
-//================================================================================================================
+UAnimInstance* FContextualAnimSceneActorData::GetAnimInstance() const
+{
+	return UContextualAnimUtilities::TryGetAnimInstance(GetActor());
+}
 
-FTransform FContextualAnimSceneActorData::GetTransform() const
+USkeletalMeshComponent* FContextualAnimSceneActorData::GetSkeletalMeshComponent() const
 {
-	FTransform Transform = FTransform::Identity;
+	return UContextualAnimUtilities::TryGetSkeletalMeshComponent(GetActor());
+}
 
-	if(Actor.IsValid())
+FAnimMontageInstance* FContextualAnimSceneActorData::GetAnimMontageInstance() const
+{
+	if (const UAnimMontage* Animation = AnimDataPtr->Animation)
 	{
-		//@TODO: Cache this during the binding
-		if (const UContextualAnimSceneActorComponent* Comp = Actor->FindComponentByClass<UContextualAnimSceneActorComponent>())
-		{
-			Transform = Comp->GetComponentTransform();
-		}
-		else
+		if (UAnimInstance* AnimInstance = GetAnimInstance())
 		{
-			Transform = Actor->GetActorTransform();
+			return AnimInstance->GetActiveInstanceForMontage(Animation);
 		}
 	}
 
-	return Transform;
+	return nullptr;
+}
+
+const UAnimMontage* FContextualAnimSceneActorData::GetAnimMontage() const
+{
+	const FAnimMontageInstance* MontageInstance = GetAnimMontageInstance();
+	return MontageInstance ? MontageInstance->Montage : nullptr;
 }
 
 float FContextualAnimSceneActorData::GetAnimTime() const
 {
-	if(const UAnimMontage* Animation = AnimDataPtr->Animation)
+	const FAnimMontageInstance* MontageInstance = GetAnimMontageInstance();
+	return MontageInstance ? MontageInstance->GetPosition() : -1.f;
+}
+
+FName FContextualAnimSceneActorData::GetCurrentSection() const
+{
+	const FAnimMontageInstance* MontageInstance = GetAnimMontageInstance();
+	return MontageInstance ? MontageInstance->GetCurrentSection() : NAME_None;
+}
+
+int32 FContextualAnimSceneActorData::GetCurrentSectionIndex() const
+{
+	if(const FAnimMontageInstance* MontageInstance = GetAnimMontageInstance())
 	{
-		if (UAnimInstance* AnimInstance = TryGetAnimInstance(Actor.Get()))
-		{
-			if(const FAnimMontageInstance* MontageInstance = AnimInstance->GetActiveInstanceForMontage(Animation))
-			{
-				return MontageInstance->GetPosition();
-			}
-		}
+		float CurrentPosition;
+		return MontageInstance->Montage->GetAnimCompositeSectionIndexFromPos(MontageInstance->GetPosition(), CurrentPosition);
 	}
-	
-	return 0.f;
+
+	return INDEX_NONE;
+}
+
+//================================================================================================================
+
+void UContextualAnimSceneInstance::BreakContextualAnimSceneActorData(const FContextualAnimSceneActorData& SceneActorData, AActor*& Actor, UAnimMontage*& Montage, float& AnimTime, int32& CurrentSectionIndex, FName& CurrentSectionName)
+{
+	Actor = SceneActorData.GetActor();
+	Montage = const_cast<UAnimMontage*>(SceneActorData.GetAnimMontage());
+	AnimTime = SceneActorData.GetAnimTime();
+	CurrentSectionIndex = SceneActorData.GetCurrentSectionIndex();
+	CurrentSectionName = SceneActorData.GetCurrentSection();
 }
 
 //================================================================================================================
@@ -87,11 +101,92 @@ UContextualAnimSceneInstance::UContextualAnimSceneInstance(const FObjectInitiali
 
 UWorld* UContextualAnimSceneInstance::GetWorld()const
 {
-	return CastChecked<UWorld>(GetOuter()->GetWorld());
+	return GetOuter() ? GetOuter()->GetWorld() : nullptr;
 }
 
 void UContextualAnimSceneInstance::Tick(float DeltaTime)
 {
+	UpdateTransitions(DeltaTime);
+}
+
+float UContextualAnimSceneInstance::GetResumePositionForSceneActor_Implementation(const FContextualAnimSceneActorData& SceneActorData, int32 DesiredSectionIndex) const
+{
+	// By default we resume from the same position in the section the leader is (we expect sections to have same length)
+
+	float StartTime = 0.f;
+	float EndTime = 0.f;
+	SceneActorData.GetAnimMontage()->GetSectionStartAndEndTime(DesiredSectionIndex, StartTime, EndTime);
+
+	const float TimeMaster = GetPositionInCurrentSection();
+	const float ResumePosition = StartTime + TimeMaster;
+
+	UE_LOG(LogContextualAnim, Verbose, TEXT("UContextualAnimSceneInstance::GetResumePositionForSceneActor Anim: %s DesiredSectionIndex: %d [%f %f] TimeMaster: %f ResumePosition: %f"),
+		*GetNameSafe(SceneActorData.GetAnimMontage()), DesiredSectionIndex, StartTime, EndTime, TimeMaster, ResumePosition);
+
+	return ResumePosition;
+}
+
+void UContextualAnimSceneInstance::UpdateTransitions(float DeltaTime)
+{
+	if (SceneAsset && SceneAsset->Transitions.Num() > 0)
+	{
+		FName CurrentSection = NAME_None;
+
+		if (const FContextualAnimSceneActorData* LeaderSceneActorData = FindSceneActorDataForRole(SceneAsset->GetLeaderRole()))
+		{
+			const FAnimMontageInstance* LeaderMontageInstance = LeaderSceneActorData->GetAnimMontageInstance();
+			if(LeaderMontageInstance)
+			{
+				CurrentSection = LeaderMontageInstance->GetCurrentSection();
+
+				// Attempt to resume montages that has been paused due a failed transition
+				for (auto& Pair : SceneActorMap)
+				{
+					FAnimMontageInstance* MontageInstance = Pair.Value.GetAnimMontageInstance();
+					if (MontageInstance && !MontageInstance->IsPlaying())
+					{
+						const int32 DesiredSectionIndex = MontageInstance->Montage->GetSectionIndex(CurrentSection);
+						if (DesiredSectionIndex != INDEX_NONE)
+						{
+							const float Position = GetResumePositionForSceneActor(Pair.Value, DesiredSectionIndex);
+							MontageInstance->SetPosition(Position);
+							MontageInstance->SetPlaying(true);
+						}
+					}
+				}
+
+				// Attempt to transition from current section
+				for (const FContextualAnimTransitionContainer& TransitionData : SceneAsset->Transitions)
+				{
+					if (TransitionData.FromSections.Contains(CurrentSection))
+					{
+						//@TODO: bForceTransition should not be in the SceneAsset
+						const bool bCanEnterTransition = TransitionData.bForceTransition || (TransitionData.Transition && TransitionData.Transition->CanEnterTransition(this, CurrentSection, TransitionData.ToSection));
+						if (bCanEnterTransition)
+						{
+							for (auto& Pair : SceneActorMap)
+							{
+								const bool bTransitionToResult = TransitionTo(Pair.Value, TransitionData.ToSection);
+
+								// If transition failed but we have a valid MontageInstance is usually because this montage doesn't have the desired section 
+								// In that case, just pause the montage. We will resume as soon as we have a valid section (see above)
+								if (bTransitionToResult == false)
+								{
+									if (FAnimMontageInstance* AnimMontageInstance = Pair.Value.GetAnimMontageInstance())
+									{
+										AnimMontageInstance->Pause();
+									}
+								}
+							}
+
+							// Break after find the first valid transition
+							break;
+						}
+					}
+				}
+			}
+		}
+	}
 }
 
 bool UContextualAnimSceneInstance::IsActorInThisScene(const AActor* Actor) const
@@ -111,67 +206,137 @@ bool UContextualAnimSceneInstance::IsActorInThisScene(const AActor* Actor) const
 	return false;
 }
 
-void UContextualAnimSceneInstance::Join(FContextualAnimSceneActorData& Data)
+const FContextualAnimSceneActorData* UContextualAnimSceneInstance::FindSceneActorDataForActor(const AActor* Actor) const
+{
+	if (Actor)
+	{
+		for (const auto& Pair : SceneActorMap)
+		{
+			const FContextualAnimSceneActorData& Data = Pair.Value;
+			if (Data.GetActor() == Actor)
+			{
+				return &Data;
+			}
+		}
+	}
+
+	return nullptr;
+}
+
+const FContextualAnimSceneActorData* UContextualAnimSceneInstance::FindSceneActorDataForRole(const FName& Role) const
+{
+	return SceneActorMap.Find(Role);
+}
+
+AActor* UContextualAnimSceneInstance::GetActorByRole(FName Role) const
+{
+	const FContextualAnimSceneActorData* SceneActorData = FindSceneActorDataForRole(Role);
+	return SceneActorData ? SceneActorData->GetActor() : nullptr;
+}
+
+void UContextualAnimSceneInstance::Join(FContextualAnimSceneActorData& SceneActorData)
 {
-	AActor* Actor = Data.GetActor();
-	if (UAnimInstance* AnimInstance = TryGetAnimInstance(Actor))
+	AActor* Actor = SceneActorData.GetActor();
+	if (UAnimInstance* AnimInstance = SceneActorData.GetAnimInstance())
 	{
 		if (UMotionWarpingComponent* MotionWarpingComp = Actor->FindComponentByClass<UMotionWarpingComponent>())
 		{
 			for (const auto& Pair : AlignmentSectionToScenePivotList)
 			{
 				const FName SyncPointName = Pair.Key;
-				const float SyncTime = Data.AnimDataPtr->GetSyncTimeForWarpSection(SyncPointName);
-				const FTransform AlignmentTransform = (Data.AnimDataPtr->AlignmentData.ExtractTransformAtTime(SyncPointName, SyncTime) * Pair.Value);
+				const float SyncTime = SceneActorData.GetAnimData()->GetSyncTimeForWarpSection(SyncPointName);
+				const FTransform AlignmentTransform = (SceneActorData.GetAnimData()->AlignmentData.ExtractTransformAtTime(SyncPointName, SyncTime) * Pair.Value);
 				MotionWarpingComp->AddOrUpdateSyncPoint(SyncPointName, AlignmentTransform);
 			}
 		}
 
-		UAnimMontage* Animation = Data.AnimDataPtr->Animation;
-		AnimInstance->Montage_Play(Animation, 1.f, EMontagePlayReturnType::MontageLength, Data.AnimStartTime);
+		UAnimMontage* Animation = SceneActorData.GetAnimData()->Animation;
+		AnimInstance->Montage_Play(Animation, 1.f, EMontagePlayReturnType::MontageLength, SceneActorData.GetAnimStartTime());
 
 		AnimInstance->OnPlayMontageNotifyBegin.AddUniqueDynamic(this, &UContextualAnimSceneInstance::OnNotifyBeginReceived);
 		AnimInstance->OnPlayMontageNotifyEnd.AddUniqueDynamic(this, &UContextualAnimSceneInstance::OnNotifyEndReceived);
 		AnimInstance->OnMontageBlendingOut.AddUniqueDynamic(this, &UContextualAnimSceneInstance::OnMontageBlendingOut);
 
-		if (UCharacterMovementComponent* CharacterMovementComp = Actor->FindComponentByClass<UCharacterMovementComponent>())
+		if(SceneActorData.GetAnimData()->bRequireFlyingMode)
 		{
-			CharacterMovementComp->SetMovementMode(MOVE_Flying);
+			if (UCharacterMovementComponent* CharacterMovementComp = Actor->FindComponentByClass<UCharacterMovementComponent>())
+			{
+				CharacterMovementComp->SetMovementMode(MOVE_Flying);
+			}
 		}
 	}
 
-	SetIgnoreCollisionWithOtherActors(Actor, true);
+	if(SceneAsset->bDisableCollisionBetweenActors)
+	{
+		SetIgnoreCollisionWithOtherActors(Actor, true);
+	}
+
+	SceneActorData.SceneInstancePtr = this;
+
+	if(UContextualAnimSceneActorComponent* SceneActorComp = SceneActorData.GetSceneActorComponent())
+	{
+		SceneActorComp->OnJoinedScene(&SceneActorData);
+	}
 
 	OnActorJoined.ExecuteIfBound(this, Actor);
 }
 
-void UContextualAnimSceneInstance::Leave(FContextualAnimSceneActorData& Data)
+void UContextualAnimSceneInstance::Leave(FContextualAnimSceneActorData& SceneActorData)
 {
-	AActor* Actor = Data.GetActor();
-	if (UAnimInstance* AnimInstance = TryGetAnimInstance(Actor))
+	// Check if we have an exit section and transition to it, otherwise just stop the montage
+	static const FName ExitSectionName = FName(TEXT("Exit"));
+	if(TransitionTo(SceneActorData, ExitSectionName) == false)
 	{
-		UAnimMontage* CurrentMontage = AnimInstance->GetCurrentActiveMontage();
-		check(CurrentMontage);
-
-		// Check if we have an exit section and transition to it, otherwise just stop the montage
-		// @TODO: This is temp until we add a solid way to deal with different states
-		static const FName ExitSectionName = FName(TEXT("Exit"));
-		const int32 SectionIdx = CurrentMontage->GetSectionIndex(ExitSectionName);
-		if (SectionIdx != INDEX_NONE)
+		AActor* Actor = SceneActorData.GetActor();
+		if (UAnimInstance* AnimInstance = SceneActorData.GetAnimInstance())
 		{
-			// Unbind blend out delegate for a moment so we don't get it during the transition
-			AnimInstance->OnMontageBlendingOut.RemoveDynamic(this, &UContextualAnimSceneInstance::OnMontageBlendingOut);
+			UAnimMontage* CurrentMontage = AnimInstance->GetCurrentActiveMontage();
+			if (ensureAlways(CurrentMontage))
+			{
+				AnimInstance->Montage_Stop(CurrentMontage->BlendOut.GetBlendTime(), CurrentMontage);
+			}
+		}
+	}
+}
+
+bool UContextualAnimSceneInstance::TransitionTo(FContextualAnimSceneActorData& SceneActorData, const FName& ToSectionName)
+{
+	UAnimInstance* AnimInstance = SceneActorData.GetAnimInstance();
+	if (AnimInstance == nullptr)
+	{		
+		return false;
+	}
 
-			AnimInstance->Montage_Play(CurrentMontage, 1.f);
-			AnimInstance->Montage_JumpToSection(ExitSectionName, CurrentMontage);
+	UAnimMontage* CurrentMontage = AnimInstance->GetCurrentActiveMontage();
+	if (CurrentMontage == nullptr)
+	{
+		UE_LOG(LogContextualAnim, Log, TEXT("UContextualAnimSceneInstance::TransitionTo. Actor is not playing any montage. Actor: %s ToSectionName: %s"),
+			*GetNameSafe(SceneActorData.GetActor()), *ToSectionName.ToString());
 
-			AnimInstance->OnMontageBlendingOut.AddUniqueDynamic(this, &UContextualAnimSceneInstance::OnMontageBlendingOut);
-		}
-		else
-		{
-			AnimInstance->Montage_Stop(CurrentMontage->BlendOut.GetBlendTime(), CurrentMontage);
-		}
+		return false;
+	}
+
+	const int32 SectionIdx = CurrentMontage->GetSectionIndex(ToSectionName);
+	if (SectionIdx == INDEX_NONE)
+	{
+		UE_LOG(LogContextualAnim, Log, TEXT("UContextualAnimSceneInstance::TransitionTo. Invalid Section. Actor: %s CurrentMontage: %s ToSectionName: %s"),
+			*GetNameSafe(SceneActorData.GetActor()), *GetNameSafe(CurrentMontage), *ToSectionName.ToString());
+
+		return false;
 	}
+
+	UE_LOG(LogContextualAnim, Verbose, TEXT("UContextualAnimSceneInstance::TransitionTo. Actor: %s CurrentMontage: %s ToSectionName: %s"),
+		*GetNameSafe(SceneActorData.GetActor()), *GetNameSafe(CurrentMontage), *ToSectionName.ToString());
+
+	// Unbind blend out delegate for a moment so we don't get it during the transition
+	AnimInstance->OnMontageBlendingOut.RemoveDynamic(this, &UContextualAnimSceneInstance::OnMontageBlendingOut);
+
+	AnimInstance->Montage_Play(CurrentMontage, 1.f);
+	AnimInstance->Montage_JumpToSection(ToSectionName, CurrentMontage);
+
+	AnimInstance->OnMontageBlendingOut.AddUniqueDynamic(this, &UContextualAnimSceneInstance::OnMontageBlendingOut);
+
+	return true;
 }
 
 void UContextualAnimSceneInstance::Start()
@@ -195,7 +360,7 @@ void UContextualAnimSceneInstance::Start()
 
 		AlignmentSectionToScenePivotList.Add(MakeTuple(SceneAsset->AlignmentSections[Idx].SectionName, ScenePivot));
 
-		//DrawDebugCoordinateSystem(GetWorld(), SceneOrigin.GetLocation(), SceneOrigin.Rotator(), 100.f, false, 10.f, 0, 5.f);
+		//DrawDebugCoordinateSystem(GetWorld(), ScenePivot.GetLocation(), ScenePivot.Rotator(), 100.f, false, 10.f, 0, 5.f);
 	}
 
 	for (auto& Pair : SceneActorMap)
@@ -203,7 +368,7 @@ void UContextualAnimSceneInstance::Start()
 		const FName& Role = Pair.Key;
 		FContextualAnimSceneActorData& Data = Pair.Value;
 
-		if (SceneAsset->GetJoinRuleForRole(Role) == EContextualAnimJoinRule::Default)
+		if (SceneAsset->GetTrackSettings(Role)->JoinRule == EContextualAnimJoinRule::Default)
 		{
 			Join(Data);
 		}
@@ -224,22 +389,34 @@ void UContextualAnimSceneInstance::OnMontageBlendingOut(UAnimMontage* Montage, b
 
 	for (auto& Pair : SceneActorMap)
 	{
-		if (Pair.Value.AnimDataPtr->Animation == Montage)
+		FContextualAnimSceneActorData& SceneActorData = Pair.Value;
+		if (SceneActorData.GetAnimData()->Animation == Montage)
 		{
-			AActor* Actor = Pair.Value.GetActor();
-			if (UAnimInstance* AnimInstance = TryGetAnimInstance(Actor))
+			AActor* Actor = SceneActorData.GetActor();
+			if (UAnimInstance* AnimInstance = SceneActorData.GetAnimInstance())
 			{
 				AnimInstance->OnPlayMontageNotifyBegin.RemoveDynamic(this, &UContextualAnimSceneInstance::OnNotifyBeginReceived);
 				AnimInstance->OnPlayMontageNotifyEnd.RemoveDynamic(this, &UContextualAnimSceneInstance::OnNotifyEndReceived);
 				AnimInstance->OnMontageBlendingOut.RemoveDynamic(this, &UContextualAnimSceneInstance::OnMontageBlendingOut);
 
-				if (UCharacterMovementComponent* CharacterMovementComp = Actor->FindComponentByClass<UCharacterMovementComponent>())
+				if (SceneActorData.GetAnimData()->bRequireFlyingMode)
 				{
-					CharacterMovementComp->SetMovementMode(MOVE_Walking);
+					if (UCharacterMovementComponent* CharacterMovementComp = Actor->FindComponentByClass<UCharacterMovementComponent>())
+					{
+						CharacterMovementComp->SetMovementMode(MOVE_Walking);
+					}
 				}
 			}
 
-			SetIgnoreCollisionWithOtherActors(Pair.Value.GetActor(), false);
+			if (SceneAsset->bDisableCollisionBetweenActors)
+			{
+				SetIgnoreCollisionWithOtherActors(SceneActorData.GetActor(), false);
+			}
+
+			if (UContextualAnimSceneActorComponent* SceneActorComp = Pair.Value.GetSceneActorComponent())
+			{
+				SceneActorComp->OnLeftScene(&SceneActorData);
+			}
 
 			OnActorLeft.ExecuteIfBound(this, Actor);
 
@@ -251,9 +428,9 @@ void UContextualAnimSceneInstance::OnMontageBlendingOut(UAnimMontage* Montage, b
 	for (auto& Pair : SceneActorMap)
 	{
 		const FContextualAnimSceneActorData& Data = Pair.Value;
-		if (UAnimInstance* AnimInstance = TryGetAnimInstance(Data.GetActor()))
+		if (UAnimInstance* AnimInstance = Data.GetAnimInstance())
 		{
-			if(AnimInstance->Montage_IsPlaying(Data.AnimDataPtr->Animation))
+			if(AnimInstance->Montage_IsPlaying(Data.GetAnimData()->Animation))
 			{
 				bShouldEnd = false;
 				break;
@@ -281,7 +458,7 @@ void UContextualAnimSceneInstance::OnNotifyBeginReceived(FName NotifyName, const
 			const FName& Role = Pair.Key;
 			FContextualAnimSceneActorData& Data = Pair.Value;
 
-			if (SceneAsset->GetJoinRuleForRole(Role) == EContextualAnimJoinRule::Late)
+			if (SceneAsset->GetTrackSettings(Role)->JoinRule == EContextualAnimJoinRule::Late)
 			{
 				Join(Data);
 			}
@@ -308,4 +485,66 @@ void UContextualAnimSceneInstance::SetIgnoreCollisionWithOtherActors(AActor* Act
 			}
 		}
 	}
-}
\ No newline at end of file
+}
+
+float UContextualAnimSceneInstance::GetCurrentSectionTimeLeft() const
+{
+	if(SceneAsset)
+	{
+		if(const FContextualAnimSceneActorData* Data = FindSceneActorDataForRole(SceneAsset->GetLeaderRole()))
+		{
+			if (const FAnimMontageInstance* MontageInstance = Data->GetAnimMontageInstance())
+			{
+				return MontageInstance->Montage->GetSectionTimeLeftFromPos(MontageInstance->GetPosition());
+			}
+		}
+	}
+	
+	return 0.f;
+}
+
+float UContextualAnimSceneInstance::GetPositionInCurrentSection() const
+{
+	if (SceneAsset)
+	{
+		if (const FContextualAnimSceneActorData* Data = FindSceneActorDataForRole(SceneAsset->GetLeaderRole()))
+		{
+			if (const FAnimMontageInstance* MontageInstance = Data->GetAnimMontageInstance())
+			{
+				float PosInSection = 0.f;
+				MontageInstance->Montage->GetAnimCompositeSectionIndexFromPos(MontageInstance->GetPosition(), PosInSection);
+				return PosInSection;
+			}
+		}
+	}
+
+	return 0.f;
+}
+
+bool UContextualAnimSceneInstance::DidCurrentSectionLoop() const
+{
+	if (SceneAsset)
+	{
+		if (const FContextualAnimSceneActorData* Data = FindSceneActorDataForRole(SceneAsset->GetLeaderRole()))
+		{
+			if (const FAnimMontageInstance* MontageInstance = Data->GetAnimMontageInstance())
+			{
+				const float PreviousPos = MontageInstance->GetPreviousPosition();
+				const float CurrentPos = MontageInstance->GetPosition();
+
+				const int32 SectionIdxPreviousPos = MontageInstance->Montage->GetSectionIndexFromPosition(PreviousPos);
+				const int32 SectionIdxCurrentPos = MontageInstance->Montage->GetSectionIndexFromPosition(CurrentPos);
+
+				if(SectionIdxPreviousPos == SectionIdxCurrentPos)
+				{
+					const float TimeLeftFromPreviousPos = MontageInstance->Montage->GetSectionTimeLeftFromPos(PreviousPos);
+					const float TimeLeftFromCurrentPos = MontageInstance->Montage->GetSectionTimeLeftFromPos(CurrentPos);
+
+					return TimeLeftFromPreviousPos < TimeLeftFromCurrentPos;
+				}
+			}
+		}
+	}
+
+	return false;
+}
diff --git a/Engine/Plugins/Experimental/Animation/ContextualAnimation/Source/ContextualAnimation/Private/ContextualAnimScenePivotProvider.cpp b/Engine/Plugins/Experimental/Animation/ContextualAnimation/Source/ContextualAnimation/Private/ContextualAnimScenePivotProvider.cpp
--- a/Engine/Plugins/Experimental/Animation/ContextualAnimation/Source/ContextualAnimation/Private/ContextualAnimScenePivotProvider.cpp
+++ b/Engine/Plugins/Experimental/Animation/ContextualAnimation/Source/ContextualAnimation/Private/ContextualAnimScenePivotProvider.cpp
@@ -25,19 +25,19 @@ UContextualAnimScenePivotProvider_Default::UContextualAnimScenePivotProvider_Def
 {
 }
 
-FTransform UContextualAnimScenePivotProvider_Default::CalculateScenePivot_Source() const
+FTransform UContextualAnimScenePivotProvider_Default::CalculateScenePivot_Source(int32 AnimDataIndex) const
 {
 	FTransform ScenePivot = FTransform::Identity;
 
 	if(const UContextualAnimSceneAsset* SceneAsset = Cast<UContextualAnimSceneAsset>(GetSceneAsset()))
 	{
-		if(const FContextualAnimTrack* PrimaryTrack = SceneAsset->DataContainer.Find(PrimaryRole))
+		if(const FContextualAnimCompositeTrack* PrimaryTrack = SceneAsset->DataContainer.Find(PrimaryRole))
 		{
-			if(const FContextualAnimTrack* SecondaryTrack = SceneAsset->DataContainer.Find(SecondaryRole))
+			if(const FContextualAnimCompositeTrack* SecondaryTrack = SceneAsset->DataContainer.Find(SecondaryRole))
 			{
-				const FTransform PrimaryTransform = SceneAsset->ExtractTransformFromAnimData(PrimaryTrack->AnimData, 0.f);
-				const FTransform SecondaryTransform = SceneAsset->ExtractTransformFromAnimData(SecondaryTrack->AnimData, 0.f);
-
+				const FTransform PrimaryTransform = PrimaryTrack->GetRootTransformForAnimDataAtIndex(AnimDataIndex);
+				const FTransform SecondaryTransform = SecondaryTrack->GetRootTransformForAnimDataAtIndex(AnimDataIndex);
+				
 				ScenePivot.SetLocation(FMath::Lerp<FVector>(PrimaryTransform.GetLocation(), SecondaryTransform.GetLocation(), Weight));
 				ScenePivot.SetRotation((SecondaryTransform.GetLocation() - PrimaryTransform.GetLocation()).GetSafeNormal2D().ToOrientationQuat());
 			}
@@ -51,15 +51,12 @@ FTransform UContextualAnimScenePivotProvider_Default::CalculateScenePivot_Runtim
 {
 	FTransform ScenePivot = FTransform::Identity;
 
-	if(const UContextualAnimSceneAsset* SceneAsset = Cast<UContextualAnimSceneAsset>(GetSceneAsset()))
+	if (const AActor* PrimaryActor = SceneActorMap.Find(PrimaryRole)->GetActor())
 	{
-		if(const AActor* PrimaryActor = SceneActorMap.Find(PrimaryRole)->GetActor())
+		if (const AActor* SecondaryActor = SceneActorMap.Find(SecondaryRole)->GetActor())
 		{
-			if (const AActor* SecondaryActor = SceneActorMap.Find(SecondaryRole)->GetActor())
-			{
-				ScenePivot.SetLocation(FMath::Lerp<FVector>(PrimaryActor->GetActorLocation(), SecondaryActor->GetActorLocation(), Weight));
-				ScenePivot.SetRotation((SecondaryActor->GetActorLocation() - PrimaryActor->GetActorLocation()).GetSafeNormal2D().ToOrientationQuat());
-			}
+			ScenePivot.SetLocation(FMath::Lerp<FVector>(PrimaryActor->GetActorLocation(), SecondaryActor->GetActorLocation(), Weight));
+			ScenePivot.SetRotation((SecondaryActor->GetActorLocation() - PrimaryActor->GetActorLocation()).GetSafeNormal2D().ToOrientationQuat());
 		}
 	}
 
@@ -74,15 +71,15 @@ UContextualAnimScenePivotProvider_RelativeTo::UContextualAnimScenePivotProvider_
 {
 }
 
-FTransform UContextualAnimScenePivotProvider_RelativeTo::CalculateScenePivot_Source() const
+FTransform UContextualAnimScenePivotProvider_RelativeTo::CalculateScenePivot_Source(int32 AnimDataIndex) const
 {
 	FTransform ScenePivot = FTransform::Identity;
 
 	if (const UContextualAnimSceneAsset* SceneAsset = Cast<UContextualAnimSceneAsset>(GetSceneAsset()))
 	{
-		if(const FContextualAnimTrack* RelativeToTrack = SceneAsset->DataContainer.Find(RelativeToRole))
+		if(const FContextualAnimCompositeTrack* Track = SceneAsset->DataContainer.Find(RelativeToRole))
 		{
-			ScenePivot = SceneAsset->ExtractTransformFromAnimData(RelativeToTrack->AnimData, 0.f);
+			ScenePivot = Track->GetRootTransformForAnimDataAtIndex(AnimDataIndex);
 		}
 	}
 
@@ -93,14 +90,11 @@ FTransform UContextualAnimScenePivotProvider_RelativeTo::CalculateScenePivot_Run
 {
 	FTransform ScenePivot = FTransform::Identity;
 
-	if (const UContextualAnimSceneAsset* SceneAsset = Cast<UContextualAnimSceneAsset>(GetSceneAsset()))
+	if (const FContextualAnimSceneActorData* Data = SceneActorMap.Find(RelativeToRole))
 	{
-		if(const FContextualAnimSceneActorData* Data = SceneActorMap.Find(RelativeToRole))
+		if (const AActor* Actor = Data->GetActor())
 		{
-			if(const AActor* Actor = Data->GetActor())
-			{
-				ScenePivot = Actor->GetActorTransform();
-			}
+			ScenePivot = Actor->GetActorTransform();
 		}
 	}
 
diff --git a/Engine/Plugins/Experimental/Animation/ContextualAnimation/Source/ContextualAnimation/Private/ContextualAnimTransition.cpp b/Engine/Plugins/Experimental/Animation/ContextualAnimation/Source/ContextualAnimation/Private/ContextualAnimTransition.cpp
new file mode 100644
--- /dev/null
+++ b/Engine/Plugins/Experimental/Animation/ContextualAnimation/Source/ContextualAnimation/Private/ContextualAnimTransition.cpp
@@ -0,0 +1,8 @@
+// Copyright Epic Games, Inc. All Rights Reserved.
+
+#include "ContextualAnimTransition.h"
+
+UContextualAnimTransition::UContextualAnimTransition(const FObjectInitializer& ObjectInitializer)
+	: Super(ObjectInitializer)
+{
+}
\ No newline at end of file
diff --git a/Engine/Plugins/Experimental/Animation/ContextualAnimation/Source/ContextualAnimation/Private/ContextualAnimTypes.cpp b/Engine/Plugins/Experimental/Animation/ContextualAnimation/Source/ContextualAnimation/Private/ContextualAnimTypes.cpp
--- a/Engine/Plugins/Experimental/Animation/ContextualAnimation/Source/ContextualAnimation/Private/ContextualAnimTypes.cpp
+++ b/Engine/Plugins/Experimental/Animation/ContextualAnimation/Source/ContextualAnimation/Private/ContextualAnimTypes.cpp
@@ -5,10 +5,15 @@
 #include "Animation/AnimMontage.h"
 #include "Containers/ArrayView.h"
 #include "AnimNotifyState_MotionWarping.h"
+#include "ContextualAnimUtilities.h"
 #include "RootMotionModifier.h"
 
 DEFINE_LOG_CATEGORY(LogContextualAnim);
 
+#if !(UE_BUILD_SHIPPING || UE_BUILD_TEST)
+TAutoConsoleVariable<int32> FContextualAnimCVars::CVarDrawDebugIKTargets(TEXT("a.ContextualAnim.DrawDebugIKTargets"), 0, TEXT("Draw Debug IK Targets"), ECVF_Cheat);
+#endif
+
 const FContextualAnimData FContextualAnimData::EmptyAnimData = FContextualAnimData();
 
 // FContextualAnimAlignmentTrackContainer
@@ -167,4 +172,22 @@ float FContextualAnimData::FindBestAnimStartTime(const FVector& LocalLocation) c
 	}
 
 	return BestTime;
+}
+
+FTransform FContextualAnimCompositeTrack::GetRootTransformForAnimDataAtIndex(int32 Index) const
+{
+	if(AnimDataContainer.IsValidIndex(Index))
+	{
+		const FContextualAnimData& AnimData = AnimDataContainer[Index];
+		if (AnimData.Animation)
+		{
+			return (Settings.MeshToComponent.Inverse() * (UContextualAnimUtilities::ExtractRootTransformFromAnimation(AnimData.Animation, 0.f) * AnimData.MeshToScene));
+		}
+		else
+		{
+			return Settings.MeshToComponent.Inverse() * AnimData.MeshToScene;
+		}
+	}
+
+	return FTransform::Identity;
 }
\ No newline at end of file
diff --git a/Engine/Plugins/Experimental/Animation/ContextualAnimation/Source/ContextualAnimation/Private/ContextualAnimUtilities.cpp b/Engine/Plugins/Experimental/Animation/ContextualAnimation/Source/ContextualAnimation/Private/ContextualAnimUtilities.cpp
--- a/Engine/Plugins/Experimental/Animation/ContextualAnimation/Source/ContextualAnimation/Private/ContextualAnimUtilities.cpp
+++ b/Engine/Plugins/Experimental/Animation/ContextualAnimation/Source/ContextualAnimation/Private/ContextualAnimUtilities.cpp
@@ -12,6 +12,8 @@
 #include "ContextualAnimTypes.h"
 #include "ContextualAnimSceneAsset.h"
 #include "Misc/MemStack.h"
+#include "GameFramework/Character.h"
+#include "ContextualAnimActorInterface.h"
 
 void UContextualAnimUtilities::ExtractLocalSpacePose(const UAnimSequenceBase* Animation, const FBoneContainer& BoneContainer, float Time, bool bExtractRootMotion, FCompactPose& OutPose)
 {
@@ -57,6 +59,27 @@ FTransform UContextualAnimUtilities::ExtractRootMotionFromAnimation(const UAnimS
 	return FTransform::Identity;
 }
 
+FTransform UContextualAnimUtilities::ExtractRootTransformFromAnimation(const UAnimSequenceBase* Animation, float Time)
+{
+	if (const UAnimMontage* AnimMontage = Cast<UAnimMontage>(Animation))
+	{
+		if (const FAnimSegment* Segment = AnimMontage->SlotAnimTracks[0].AnimTrack.GetSegmentAtTime(Time))
+		{
+			if (const UAnimSequence* AnimSequence = Cast<UAnimSequence>(Segment->AnimReference))
+			{
+				const float AnimSequenceTime = Segment->ConvertTrackPosToAnimPos(Time);
+				return AnimSequence->ExtractRootTrackTransform(AnimSequenceTime, nullptr);
+			}
+		}
+	}
+	else if (const UAnimSequence* AnimSequence = Cast<UAnimSequence>(Animation))
+	{
+		return AnimSequence->ExtractRootTrackTransform(Time, nullptr);
+	}
+
+	return FTransform::Identity;
+}
+
 void UContextualAnimUtilities::DrawDebugPose(const UWorld* World, const UAnimSequenceBase* Animation, float Time, const FTransform& LocalToWorldTransform, const FColor& Color, float LifeTime, float Thickness)
 { 
 	if(World)
@@ -100,22 +123,76 @@ void UContextualAnimUtilities::DrawDebugPose(const UWorld* World, const UAnimSeq
 	}
 }
 
-void UContextualAnimUtilities::DrawDebugScene(const UWorld* World, const UContextualAnimSceneAsset* SceneAsset, float Time, const FTransform& ToWorldTransform, const FColor& Color, float LifeTime, float Thickness)
+void UContextualAnimUtilities::DrawDebugScene(const UWorld* World, const UContextualAnimSceneAsset* SceneAsset, int32 AnimDataIndex, float Time, const FTransform& ToWorldTransform, const FColor& Color, float LifeTime, float Thickness)
 {
 	if (World && SceneAsset)
 	{
 		for(const auto& Pair : SceneAsset->DataContainer)
 		{
-			const FContextualAnimTrack& Track = Pair.Value;
-			const FTransform Transform = (SceneAsset->MeshToComponent * Track.AnimData.GetAlignmentTransformAtTime(Time)) * ToWorldTransform;
-			if(const UAnimMontage* Animation = Track.AnimData.Animation)
-			{
-				DrawDebugPose(World, Animation, Time, Transform, Color, LifeTime, Thickness);
-			}
-			else
+			if (Pair.Value.AnimDataContainer.IsValidIndex(AnimDataIndex))
 			{
-				DrawDebugCoordinateSystem(World, Transform.GetLocation(), Transform.Rotator(), 50.f, false, LifeTime, 0, Thickness);
+				const FContextualAnimData& AnimData = Pair.Value.AnimDataContainer[AnimDataIndex];
+				const FTransform Transform = (Pair.Value.Settings.MeshToComponent * AnimData.GetAlignmentTransformAtTime(Time)) * ToWorldTransform;
+				if (const UAnimMontage* Animation = AnimData.Animation)
+				{
+					DrawDebugPose(World, Animation, Time, Transform, Color, LifeTime, Thickness);
+				}
+				else
+				{
+					DrawDebugCoordinateSystem(World, Transform.GetLocation(), Transform.Rotator(), 50.f, false, LifeTime, 0, Thickness);
+				}
 			}
 		}
 	}
+}
+
+USkeletalMeshComponent* UContextualAnimUtilities::TryGetSkeletalMeshComponent(AActor* Actor)
+{
+	USkeletalMeshComponent* SkelMeshComp = nullptr;
+	if (Actor)
+	{
+		if (ACharacter* Character = Cast<ACharacter>(Actor))
+		{
+			SkelMeshComp = Character->GetMesh();
+		}
+		else if (Actor->GetClass()->ImplementsInterface(UContextualAnimActorInterface::StaticClass()))
+		{
+			SkelMeshComp = IContextualAnimActorInterface::Execute_GetMesh(Actor);
+		}
+		else
+		{
+			SkelMeshComp = Actor->FindComponentByClass<USkeletalMeshComponent>();
+		}
+	}
+
+	return SkelMeshComp;
+}
+
+UAnimInstance* UContextualAnimUtilities::TryGetAnimInstance(AActor* Actor)
+{
+	if (USkeletalMeshComponent* SkelMeshComp = UContextualAnimUtilities::TryGetSkeletalMeshComponent(Actor))
+	{
+		return SkelMeshComp->GetAnimInstance();
+	}
+
+	return nullptr;
+}
+
+void UContextualAnimUtilities::BP_Montage_GetSectionStartAndEndTime(const UAnimMontage* Montage, int32 SectionIndex, float& OutStartTime, float& OutEndTime)
+{
+	if(Montage)
+	{
+		Montage->GetSectionStartAndEndTime(SectionIndex, OutStartTime, OutEndTime);
+	}
+}
+
+float UContextualAnimUtilities::BP_Montage_GetSectionTimeLeftFromPos(const UAnimMontage* Montage, float Position)
+{
+	//UAnimMontage::GetSectionTimeLeftFromPos is not const :(
+	return Montage ? (const_cast<UAnimMontage*>(Montage))->GetSectionTimeLeftFromPos(Position) : -1.f;
+}
+
+float UContextualAnimUtilities::BP_Montage_GetSectionLength(const UAnimMontage* Montage, int32 SectionIndex)
+{
+	return Montage ? Montage->GetSectionLength(SectionIndex) : -1.f;
 }
\ No newline at end of file
diff --git a/Engine/Plugins/Experimental/Animation/ContextualAnimation/Source/ContextualAnimation/Public/ContextualAnimCompositeSceneAsset.h b/Engine/Plugins/Experimental/Animation/ContextualAnimation/Source/ContextualAnimation/Public/ContextualAnimCompositeSceneAsset.h
--- a/Engine/Plugins/Experimental/Animation/ContextualAnimation/Source/ContextualAnimation/Public/ContextualAnimCompositeSceneAsset.h
+++ b/Engine/Plugins/Experimental/Animation/ContextualAnimation/Source/ContextualAnimation/Public/ContextualAnimCompositeSceneAsset.h
@@ -5,65 +5,7 @@
 #include "ContextualAnimSceneAssetBase.h"
 #include "ContextualAnimCompositeSceneAsset.generated.h"
 
-/** Stores the result of a query function */
-USTRUCT(BlueprintType)
-struct FContextualAnimQueryResult
-{
-	GENERATED_BODY()
-
-	UPROPERTY(BlueprintReadWrite, Category = "Defaults")
-	TWeakObjectPtr<UAnimMontage> Animation;
-
-	UPROPERTY(BlueprintReadWrite, Category = "Defaults")
-	FTransform EntryTransform;
-
-	UPROPERTY(BlueprintReadWrite, Category = "Defaults")
-	FTransform SyncTransform;
-
-	UPROPERTY(BlueprintReadWrite, Category = "Defaults")
-	float AnimStartTime = 0.f;
-
-	UPROPERTY(BlueprintReadWrite, Category = "Defaults")
-	int32 DataIndex = INDEX_NONE;
-
-	void Reset()
-	{
-		Animation.Reset();
-		EntryTransform = SyncTransform = FTransform::Identity;
-		AnimStartTime = 0.f;
-		DataIndex = INDEX_NONE;
-	}
-
-	FORCEINLINE bool IsValid() const { return DataIndex != INDEX_NONE; }
-};
-
-/** Stores the parameters passed into query function */
-USTRUCT(BlueprintType)
-struct FContextualAnimQueryParams
-{
-	GENERATED_BODY()
-
-	UPROPERTY(EditAnywhere, BlueprintReadWrite, Category = "Defaults")
-	TWeakObjectPtr<const AActor> Querier;
-
-	UPROPERTY(EditAnywhere, BlueprintReadWrite, Category = "Defaults")
-	FTransform QueryTransform;
-
-	UPROPERTY(EditAnywhere, BlueprintReadWrite, Category = "Defaults")
-	bool bComplexQuery = false;
-
-	UPROPERTY(EditAnywhere, BlueprintReadWrite, Category = "Defaults")
-	bool bFindAnimStartTime = false;
-
-	FContextualAnimQueryParams() {}
-
-	FContextualAnimQueryParams(const AActor* InQuerier, bool bInComplexQuery, bool bInFindAnimStartTime)
-		: Querier(InQuerier), bComplexQuery(bInComplexQuery), bFindAnimStartTime(bInFindAnimStartTime) {}
-
-	FContextualAnimQueryParams(const FTransform& InQueryTransform, bool bInComplexQuery, bool bInFindAnimStartTime)
-		: QueryTransform(InQueryTransform), bComplexQuery(bInComplexQuery), bFindAnimStartTime(bInFindAnimStartTime) {}
-};
-
+// DEPRECATED will be removed shortly
 UCLASS()
 class CONTEXTUALANIMATION_API UContextualAnimCompositeSceneAsset : public UContextualAnimSceneAssetBase
 {
@@ -82,14 +24,18 @@ public:
 
 	UContextualAnimCompositeSceneAsset(const FObjectInitializer& ObjectInitializer);
 
-	PRAGMA_DISABLE_DEPRECATION_WARNINGS // Suppress compiler warning on override of deprecated function
-	UE_DEPRECATED(5.0, "Use version that takes FObjectPreSaveContext instead.")
-	virtual void PreSave(const class ITargetPlatform* TargetPlatform) override;
-	PRAGMA_ENABLE_DEPRECATION_WARNINGS
 	virtual void PreSave(FObjectPreSaveContext ObjectSaveContext) override;
 
-	bool QueryData(FContextualAnimQueryResult& OutResult, const FContextualAnimQueryParams& QueryParams, const FTransform& ToWorldTransform) const;
+	virtual const FContextualAnimTrackSettings* GetTrackSettings(const FName& Role) const override;
+
+	virtual const FContextualAnimData* GetAnimDataForRoleAtIndex(const FName& Role, int32 Index) const override;
+
+	virtual void ForEachAnimData(FForEachAnimDataFunction Function) const override;
 
-	virtual UClass* GetPreviewActorClassForRole(const FName& Role) const override;
-	virtual EContextualAnimJoinRule GetJoinRuleForRole(const FName& Role) const override;
+	//@TODO: Temp until have a roles asset
+	virtual TArray<FName> GetRoles() const override;
+
+	virtual bool Query(const FName& Role, FContextualAnimQueryResult& OutResult, const FContextualAnimQueryParams& QueryParams, const FTransform& ToWorldTransform) const override;
+
+	bool QueryData(FContextualAnimQueryResult& OutResult, const FContextualAnimQueryParams& QueryParams, const FTransform& ToWorldTransform) const;
 };
diff --git a/Engine/Plugins/Experimental/Animation/ContextualAnimation/Source/ContextualAnimation/Public/ContextualAnimManager.h b/Engine/Plugins/Experimental/Animation/ContextualAnimation/Source/ContextualAnimation/Public/ContextualAnimManager.h
--- a/Engine/Plugins/Experimental/Animation/ContextualAnimation/Source/ContextualAnimation/Public/ContextualAnimManager.h
+++ b/Engine/Plugins/Experimental/Animation/ContextualAnimation/Source/ContextualAnimation/Public/ContextualAnimManager.h
@@ -4,14 +4,18 @@
 
 #include "CoreMinimal.h"
 #include "UObject/Object.h"
+#include "Tickable.h"
+#include "ContextualAnimTypes.h"
 #include "ContextualAnimManager.generated.h"
 
+class UContextualAnimSceneActorComponent;
 class UContextualAnimSceneAssetBase;
 class UContextualAnimSceneInstance;
+class AActor;
 class UWorld;
 
 UCLASS()
-class CONTEXTUALANIMATION_API UContextualAnimManager : public UObject
+class CONTEXTUALANIMATION_API UContextualAnimManager : public UObject, public FTickableGameObject
 {
 	GENERATED_BODY()
 
@@ -21,20 +25,52 @@ public:
 
 	virtual UWorld* GetWorld() const override;
 
-	void Tick(float DeltaTime);
+	// FTickableGameObject begin
+	virtual UWorld* GetTickableGameObjectWorld() const override { return GetWorld(); }
+	virtual void Tick(float DeltaTime) override;
+	virtual ETickableTickType GetTickableTickType() const override;
+	virtual TStatId GetStatId() const override;
+	virtual bool IsTickableInEditor() const { return true; }
+	// FTickableGameObject end
 
-	bool TryStartScene(const UContextualAnimSceneAssetBase* SceneAsset, const TMap<FName, AActor*>& Bindings);
+	static UContextualAnimManager* Get(const UWorld* World);
+
+	void RegisterSceneActorComponent(UContextualAnimSceneActorComponent* SceneActorComp);
+
+	void UnregisterSceneActorComponent(UContextualAnimSceneActorComponent* SceneActorComp);
+
+	/** Attempt to start an scene instance with the supplied bindings for each role */
+	bool TryStartScene(const UContextualAnimSceneAssetBase* SceneAsset, const FContextualAnimSceneBindings& Bindings);
+
+	/** Attempt to start an scene instance with PrimaryActor bound to the primary role and the first component valid for each of the other roles */
+	bool TryStartScene(const UContextualAnimSceneAssetBase* SceneAsset, AActor* PrimaryActor, const TSet<UContextualAnimSceneActorComponent*>& SceneActorComps);
 
+	UFUNCTION(BlueprintCallable, Category = "Contextual Anim|Manager")
+	UContextualAnimSceneActorComponent* FindClosestSceneActorCompToActor(const AActor* Actor) const;
+	
+	UFUNCTION(BlueprintCallable, Category = "Contextual Anim|Manager")
 	bool TryStopSceneWithActor(AActor* Actor);
 
+	UFUNCTION(BlueprintCallable, Category = "Contextual Anim|Manager")
 	bool IsActorInAnyScene(AActor* Actor) const;
 
+	UFUNCTION(BlueprintCallable, Category = "Contextual Anim|Manager")
 	UContextualAnimSceneInstance* GetSceneWithActor(AActor* Actor);
 
-	static UContextualAnimManager* Get(const UWorld* World);
+	UFUNCTION(BlueprintPure, Category = "Contextual Anim|Manager", meta = (WorldContext = "WorldContextObject"))
+	static UContextualAnimManager* GetContextualAnimManager(UObject* WorldContextObject);
+
+	FORCEINLINE const TSet<UContextualAnimSceneActorComponent*>& GetSceneActorCompContainer() const { return SceneActorCompContainer; };
+
+	UFUNCTION(BlueprintCallable, Category = "Contextual Anim|Manager", meta = (DisplayName = "TryStartSceneWithBindings"))
+	bool BP_TryStartSceneWithBindings(const UContextualAnimSceneAssetBase* SceneAsset, const FContextualAnimSceneBindings& Bindings) { return TryStartScene(SceneAsset, Bindings); }
 
 protected:
 
+	/** Container for all SceneActorComps in the world */
+	UPROPERTY()
+	TSet<UContextualAnimSceneActorComponent*> SceneActorCompContainer;
+
 	UPROPERTY()
 	TArray<UContextualAnimSceneInstance*> Instances;
 
diff --git a/Engine/Plugins/Experimental/Animation/ContextualAnimation/Source/ContextualAnimation/Public/ContextualAnimMetadata.h b/Engine/Plugins/Experimental/Animation/ContextualAnimation/Source/ContextualAnimation/Public/ContextualAnimMetadata.h
--- a/Engine/Plugins/Experimental/Animation/ContextualAnimation/Source/ContextualAnimation/Public/ContextualAnimMetadata.h
+++ b/Engine/Plugins/Experimental/Animation/ContextualAnimation/Source/ContextualAnimation/Public/ContextualAnimMetadata.h
@@ -9,33 +9,6 @@
 
 struct FContextualAnimData;
 
-USTRUCT(BlueprintType)
-struct FContextualAnimDistanceParam
-{
-	GENERATED_BODY()
-
-	UPROPERTY(VisibleAnywhere, BlueprintReadWrite, Category = "Defaults", meta = (UIMin = 0, ClampMin = 0))
-	float Value = 0.f;
-
-	UPROPERTY(EditAnywhere, BlueprintReadWrite, Category = "Defaults", meta = (UIMin = 0, ClampMin = 0))
-	float MinDistance = 0.f;
-
-	UPROPERTY(EditAnywhere, BlueprintReadWrite, Category = "Defaults", meta = (UIMin = 0, ClampMin = 0))
-	float MaxDistance = 0.f;
-};
-
-USTRUCT(BlueprintType)
-struct FContextualAnimAngleParam
-{
-	GENERATED_BODY()
-
-	UPROPERTY(VisibleAnywhere, BlueprintReadWrite, Category = "Defaults")
-	float Value = 0.f;
-
-	UPROPERTY(EditAnywhere, BlueprintReadWrite, Category = "Defaults")
-	float Tolerance = 0.f;
-};
-
 UCLASS(BlueprintType, EditInlineNew)
 class CONTEXTUALANIMATION_API UContextualAnimMetadata : public UObject
 {
@@ -44,18 +17,27 @@ class CONTEXTUALANIMATION_API UContextualAnimMetadata : public UObject
 public:
 
 	UPROPERTY(EditAnywhere, BlueprintReadWrite, Category = "Defaults")
-	float OffsetFromOrigin = 0.f;
+	FVector2D OriginOffset;
 
-	UPROPERTY(EditAnywhere, BlueprintReadWrite, Category = "Defaults")
-	FContextualAnimDistanceParam Distance;
+	UPROPERTY(EditAnywhere, BlueprintReadWrite, Category = "Defaults", meta = (ClampMin = "-180", ClampMax = "180", UIMin = "-180", UIMax = "180"))
+	float DirectionOffset = 0.f;
 
-	UPROPERTY(EditAnywhere, BlueprintReadWrite, Category = "Defaults")
-	FContextualAnimAngleParam Angle;
+	UPROPERTY(EditAnywhere, BlueprintReadWrite, Category = "Defaults", meta = (ClampMin = "0", UIMin = "0"))
+	float MaxDistance = 0.f;
 
-	UPROPERTY(EditAnywhere, BlueprintReadWrite, Category = "Defaults")
-	FContextualAnimAngleParam Facing;
+	UPROPERTY(EditAnywhere, BlueprintReadWrite, Category = "Defaults", meta = (ClampMin = "0", UIMin = "0"))
+	float NearWidth = 0.f;
+
+	UPROPERTY(EditAnywhere, BlueprintReadWrite, Category = "Defaults", meta = (ClampMin = "0", UIMin = "0"))
+	float FarWidth = 0.f;
+
+	UPROPERTY(EditAnywhere, BlueprintReadWrite, Category = "Defaults", meta = (ClampMin = "-180", ClampMax = "180", UIMin = "-180", UIMax = "180"))
+	float Facing = 0.f;
 
 	UContextualAnimMetadata(const FObjectInitializer& ObjectInitializer);
 
 	class UContextualAnimSceneAssetBase* GetSceneAssetOwner() const;
+
+	bool DoesQuerierPassConditions(const FContextualAnimQuerier& Querier, const FContextualAnimQueryContext& Context, const FTransform& EntryTransform) const;
+
 };
\ No newline at end of file
diff --git a/Engine/Plugins/Experimental/Animation/ContextualAnimation/Source/ContextualAnimation/Public/ContextualAnimSceneActorComponent.h b/Engine/Plugins/Experimental/Animation/ContextualAnimation/Source/ContextualAnimation/Public/ContextualAnimSceneActorComponent.h
--- a/Engine/Plugins/Experimental/Animation/ContextualAnimation/Source/ContextualAnimation/Public/ContextualAnimSceneActorComponent.h
+++ b/Engine/Plugins/Experimental/Animation/ContextualAnimation/Source/ContextualAnimation/Public/ContextualAnimSceneActorComponent.h
@@ -4,12 +4,14 @@
 
 #include "CoreMinimal.h"
 #include "Components/SphereComponent.h"
-#include "ContextualAnimCompositeSceneAsset.h"
+#include "ContextualAnimTypes.h"
 #include "ContextualAnimSceneActorComponent.generated.h"
 
 class UAnimInstance;
 class UAnimMontage;
 class AActor;
+class UContextualAnimSceneInstance;
+struct FContextualAnimSceneActorData;
 
 USTRUCT(BlueprintType)
 struct FContextualAnimDebugParams
@@ -23,17 +25,25 @@ struct FContextualAnimDebugParams
 	float DrawAlignmentTransformAtTime = 0.f;
 };
 
+DECLARE_DYNAMIC_MULTICAST_DELEGATE_OneParam(FContextualAnimSceneActorCompDelegate, class UContextualAnimSceneActorComponent*, SceneActorComponent);
+
 UCLASS(meta = (BlueprintSpawnableComponent))
-class CONTEXTUALANIMATION_API UContextualAnimSceneActorComponent : public USphereComponent
+class CONTEXTUALANIMATION_API UContextualAnimSceneActorComponent : public UPrimitiveComponent
 {
 	GENERATED_BODY()
 
 public:
 
-	UContextualAnimSceneActorComponent(const FObjectInitializer& ObjectInitializer);
+	/** Event that happens when the actor owner of this component joins an scene */
+	UPROPERTY(BlueprintAssignable, Category = "Events")
+	FContextualAnimSceneActorCompDelegate OnJoinedSceneDelegate;
+
+	/** Event that happens when the actor owner of this component leave an scene */
+	UPROPERTY(BlueprintAssignable, Category = "Events")
+	FContextualAnimSceneActorCompDelegate OnLeftSceneDelegate;
 
 	UPROPERTY(EditAnywhere, BlueprintReadWrite, Category = "Config")
-	UContextualAnimCompositeSceneAsset* SceneAsset;
+	class UContextualAnimSceneAssetBase* SceneAsset;
 
 	UPROPERTY(EditAnywhere, BlueprintReadWrite, Category = "Debug")
 	bool bEnableDebug;
@@ -41,9 +51,34 @@ public:
 	UPROPERTY(EditAnywhere, BlueprintReadWrite, Category = "Debug", meta = (ShowOnlyInnerProperties, EditCondition = "bEnableDebug"))
 	FContextualAnimDebugParams DebugParams;
 
-	virtual FPrimitiveSceneProxy* CreateSceneProxy() override;
+	UContextualAnimSceneActorComponent(const FObjectInitializer& ObjectInitializer);
+
+	virtual void OnRegister() override;
+	virtual void OnUnregister() override;
 	virtual FBoxSphereBounds CalcBounds(const FTransform& LocalToWorld) const;
+	virtual FPrimitiveSceneProxy* CreateSceneProxy() override;
+
+	/** Called from the scene instance when the actor owner of this component joins an scene */
+	void OnJoinedScene(const FContextualAnimSceneActorData* SceneActorData);
+	
+	/** Called from the scene instance when the actor owner of this component leave an scene */
+	void OnLeftScene(const FContextualAnimSceneActorData* SceneActorData);
+
+protected:
+
+	/** Ptr back to the scene actor data that represent us in the scene instance we are part of */
+	const FContextualAnimSceneActorData* SceneActorDataPtr = nullptr;
+
+	/** List of IKTarget for this frame */
+	UPROPERTY()
+	TArray<FContextualAnimIKTarget> IKTargets;
+
+	void UpdateIKTargets();
 
-	UFUNCTION(BlueprintCallable, Category = "Contextual Animation System")
-	bool QueryData(const FContextualAnimQueryParams& QueryParams, FContextualAnimQueryResult& Result) const;
+	/** 
+	 * Event called right before owner's mesh ticks the pose when we are in a scene instance and IK Targets are required. 
+	 * Used to update IK Targets before animation need them 
+	 */
+	UFUNCTION()
+	void OnTickPose(class USkinnedMeshComponent* SkinnedMeshComponent, float DeltaTime, bool bNeedsValidRootMotion);
 };
diff --git a/Engine/Plugins/Experimental/Animation/ContextualAnimation/Source/ContextualAnimation/Public/ContextualAnimSceneAsset.h b/Engine/Plugins/Experimental/Animation/ContextualAnimation/Source/ContextualAnimation/Public/ContextualAnimSceneAsset.h
--- a/Engine/Plugins/Experimental/Animation/ContextualAnimation/Source/ContextualAnimation/Public/ContextualAnimSceneAsset.h
+++ b/Engine/Plugins/Experimental/Animation/ContextualAnimation/Source/ContextualAnimation/Public/ContextualAnimSceneAsset.h
@@ -13,19 +13,20 @@ class CONTEXTUALANIMATION_API UContextualAnimSceneAsset : public UContextualAnim
 public:
 
 	UPROPERTY(EditAnywhere, BlueprintReadWrite, Category = "Animations")
-	TMap<FName, FContextualAnimTrack> DataContainer;
+	TMap<FName, FContextualAnimCompositeTrack> DataContainer;
 
 	UContextualAnimSceneAsset(const FObjectInitializer& ObjectInitializer);
 
-	PRAGMA_DISABLE_DEPRECATION_WARNINGS // Suppress compiler warning on override of deprecated function
-	UE_DEPRECATED(5.0, "Use version that takes FObjectPreSaveContext instead.")
-	virtual void PreSave(const class ITargetPlatform* TargetPlatform) override;
-	PRAGMA_ENABLE_DEPRECATION_WARNINGS
 	virtual void PreSave(FObjectPreSaveContext ObjectSaveContext) override;
 
-	virtual UClass* GetPreviewActorClassForRole(const FName& Role) const override;
+	virtual const FContextualAnimTrackSettings* GetTrackSettings(const FName& Role) const override;
 	
-	virtual EContextualAnimJoinRule GetJoinRuleForRole(const FName& Role) const override;
-	
-	const FContextualAnimTrack* FindTrack(const FName& Role) const;
+	virtual const FContextualAnimData* GetAnimDataForRoleAtIndex(const FName& Role, int32 Index) const override;
+
+	virtual void ForEachAnimData(FForEachAnimDataFunction Function) const override;
+
+	virtual bool Query(const FName& Role, FContextualAnimQueryResult& OutResult, const FContextualAnimQueryParams& QueryParams, const FTransform& ToWorldTransform) const override;
+
+	//@TODO: Temp until have a roles asset
+	virtual TArray<FName> GetRoles() const override;
 };
diff --git a/Engine/Plugins/Experimental/Animation/ContextualAnimation/Source/ContextualAnimation/Public/ContextualAnimSceneAssetBase.h b/Engine/Plugins/Experimental/Animation/ContextualAnimation/Source/ContextualAnimation/Public/ContextualAnimSceneAssetBase.h
--- a/Engine/Plugins/Experimental/Animation/ContextualAnimation/Source/ContextualAnimation/Public/ContextualAnimSceneAssetBase.h
+++ b/Engine/Plugins/Experimental/Animation/ContextualAnimation/Source/ContextualAnimation/Public/ContextualAnimSceneAssetBase.h
@@ -3,10 +3,12 @@
 #pragma once
 
 #include "Engine/DataAsset.h"
+#include "Templates/SubclassOf.h"
 #include "ContextualAnimTypes.h"
 #include "ContextualAnimSceneAssetBase.generated.h"
 
 class UContextualAnimScenePivotProvider;
+class UContextualAnimSceneInstance;
 
 USTRUCT(BlueprintType)
 struct FContextualAnimAlignmentSectionData
@@ -20,7 +22,13 @@ struct FContextualAnimAlignmentSectionData
 	UContextualAnimScenePivotProvider* ScenePivotProvider = nullptr;
 
 	UPROPERTY(VisibleAnywhere, BlueprintReadWrite, Category = "Alignment", AdvancedDisplay)
-	FTransform ScenePivot;
+	TArray<FTransform> ScenePivots;
+};
+
+enum class EContextualAnimForEachResult : uint8
+{
+	Break,
+	Continue,
 };
 
 UCLASS(Abstract, Blueprintable)
@@ -30,38 +38,66 @@ class CONTEXTUALANIMATION_API UContextualAnimSceneAssetBase : public UDataAsset
 
 public:
 
-	UPROPERTY(EditAnywhere, BlueprintReadWrite, Category = "Alignment")
+	typedef TFunction<EContextualAnimForEachResult(const FName& Role, const FContextualAnimData& AnimData)> FForEachAnimDataFunction;
+
+	UPROPERTY(EditAnywhere, Category = "Settings")
+	TSubclassOf<UContextualAnimSceneInstance> SceneInstanceClass;
+
+	UPROPERTY(EditAnywhere, Category = "Settings")
+	bool bDisableCollisionBetweenActors;
+
+	UPROPERTY(EditAnywhere, Category = "Alignment")
 	TArray<FContextualAnimAlignmentSectionData> AlignmentSections;
 
-	UPROPERTY(EditAnywhere, BlueprintReadWrite, Category = "Alignment", meta = (ClampMin = "1", ClampMax = "60"), AdvancedDisplay)
+	UPROPERTY(EditAnywhere, Category = "Alignment", meta = (ClampMin = "1", ClampMax = "60"), AdvancedDisplay)
 	int32 SampleRate;
 
-	UPROPERTY(EditAnywhere, BlueprintReadWrite, Category = "Alignment", AdvancedDisplay)
-	FTransform MeshToComponent;
+	UPROPERTY(EditAnywhere, Category = "Transitions")
+	TArray<FContextualAnimTransitionContainer> Transitions;
 
-	UPROPERTY(EditAnywhere, BlueprintReadWrite, Category = "Roles")
+	//@TODO: Rename this is used for selection and move to private
+	UPROPERTY(EditAnywhere, Category = "Roles")
 	FName PrimaryRole;
 
 	UContextualAnimSceneAssetBase(const FObjectInitializer& ObjectInitializer);
 
-	PRAGMA_DISABLE_DEPRECATION_WARNINGS // Suppress compiler warning on override of deprecated function
-	UE_DEPRECATED(5.0, "Use version that takes FObjectPreSaveContext instead.")
-	virtual void PreSave(const class ITargetPlatform* TargetPlatform) override;
-	PRAGMA_ENABLE_DEPRECATION_WARNINGS
 	virtual void PreSave(FObjectPreSaveContext ObjectSaveContext) override;
 
 	//@TODO: Make pure virtual
-	virtual UClass* GetPreviewActorClassForRole(const FName& Role) const { return nullptr; }
-	virtual EContextualAnimJoinRule GetJoinRuleForRole(const FName& Role) const { return EContextualAnimJoinRule::Default; }
+	virtual const FContextualAnimTrackSettings* GetTrackSettings(const FName& Role) const { return nullptr; }
+	virtual const FContextualAnimData* GetAnimDataForRoleAtIndex(const FName& Role, int32 Index) const { return nullptr; };
+	virtual void ForEachAnimData(FForEachAnimDataFunction Function) const {}
 
-	FTransform ExtractTransformFromAnimData(const FContextualAnimData& AnimData, float Time) const;
+	//@TODO: Temp until have a roles asset
+	virtual TArray<FName> GetRoles() const { return TArray<FName>(); }
 
 	FORCEINLINE FName GetAlignmentSectionNameAtIndex(int32 Index) const 
 	{
 		return AlignmentSections.IsValidIndex(Index) ? AlignmentSections[Index].SectionName : NAME_None;
 	}
 
+	FORCEINLINE const FName& GetLeaderRole() const { return LeaderRole != NAME_None ? LeaderRole : PrimaryRole; }
+
+	FORCEINLINE float GetRadius() const { return Radius; }
+
+	virtual bool Query(const FName& Role, FContextualAnimQueryResult& OutResult, const FContextualAnimQueryParams& QueryParams, const FTransform& ToWorldTransform) const { return false; };
+
+	bool QueryCompositeTrack(const FContextualAnimCompositeTrack* Track, FContextualAnimQueryResult& OutResult, const FContextualAnimQueryParams& QueryParams, const FTransform& ToWorldTransform) const;
+
 protected:
 
-	void GenerateAlignmentTracksRelativeToScenePivot(FContextualAnimData& AnimData);
+	void GenerateAlignmentTracks(const FContextualAnimTrackSettings& Settings, FContextualAnimData& AnimData);
+	
+	void GenerateIKTargetTracks(const FContextualAnimTrackSettings& Settings, FContextualAnimData& AnimData);
+
+	void UpdateRadius();
+
+private:
+
+	UPROPERTY(EditAnywhere, Category = "Roles")
+	FName LeaderRole;
+
+	/** Radius that enclose all the entry points */
+	UPROPERTY(VisibleAnywhere, Category = "Settings")
+	float Radius = 0.f;
 };
diff --git a/Engine/Plugins/Experimental/Animation/ContextualAnimation/Source/ContextualAnimation/Public/ContextualAnimSceneInstance.h b/Engine/Plugins/Experimental/Animation/ContextualAnimation/Source/ContextualAnimation/Public/ContextualAnimSceneInstance.h
--- a/Engine/Plugins/Experimental/Animation/ContextualAnimation/Source/ContextualAnimation/Public/ContextualAnimSceneInstance.h
+++ b/Engine/Plugins/Experimental/Animation/ContextualAnimation/Source/ContextualAnimation/Public/ContextualAnimSceneInstance.h
@@ -6,42 +6,82 @@
 #include "UObject/Object.h"
 #include "Animation/AnimNotifies/AnimNotify.h"
 #include "GameFramework/Actor.h"
+#include "ContextualAnimTypes.h"
 #include "ContextualAnimSceneInstance.generated.h"
 
+struct FAnimMontageInstance;
 struct FContextualAnimData;
+struct FContextualAnimTrackSettings;
+class UContextualAnimSceneInstance;
+class UContextualAnimSceneActorComponent;
 class UContextualAnimSceneAssetBase;
+class UAnimInstance;
+class USkeletalMeshComponent;
 class UAnimMontage;
 class UWorld;
 
 /** Represent an actor bound to a role in the scene */
-USTRUCT()
+USTRUCT(BlueprintType)
 struct CONTEXTUALANIMATION_API FContextualAnimSceneActorData
 {
 	GENERATED_BODY()
 
-	/** The actual actor in the world */
-	TWeakObjectPtr<AActor> Actor = nullptr;
-
-	/** Ptr to the animation data in the scene asset used by this actor */
-	const FContextualAnimData* AnimDataPtr = nullptr;
-
-	/** Desired time to start the animation */
-	float AnimStartTime = 0.f;
-
-	FContextualAnimSceneActorData() : Actor(nullptr), AnimDataPtr(nullptr), AnimStartTime(0.f) {}
+		FContextualAnimSceneActorData()
+		: Actor(nullptr), AnimDataPtr(nullptr), SettingsPtr(nullptr), AnimStartTime(0.f)
+	{}
 
-	FContextualAnimSceneActorData(AActor* InActor, const FContextualAnimData* InAnimData, float InAnimStartTime = 0.f)
-		: Actor(InActor), AnimDataPtr(InAnimData), AnimStartTime(InAnimStartTime)
+	FContextualAnimSceneActorData(AActor* InActor, const FContextualAnimData* InAnimData, const FContextualAnimTrackSettings* InSettings, float InAnimStartTime = 0.f)
+		: Actor(InActor), AnimDataPtr(InAnimData), SettingsPtr(InSettings), AnimStartTime(InAnimStartTime)
 	{}
 
 	/** Return a pointer to the actual actor in the world */
 	FORCEINLINE AActor* GetActor() const { return Actor.Get(); }
 
+	FORCEINLINE float GetAnimStartTime() const { return AnimStartTime; }
+
+	FORCEINLINE const FContextualAnimData* GetAnimData() const { return AnimDataPtr; }
+
+	FORCEINLINE const FContextualAnimTrackSettings* GetSettings() const { return SettingsPtr; }
+
+	FORCEINLINE const UContextualAnimSceneInstance* GetSceneInstance() const { return SceneInstancePtr.Get(); }
+
 	/** Return the transform used for alignment for this scene actor */
 	FTransform GetTransform() const;
 
 	/** Return the current playback time of the animation this actor is playing */
 	float GetAnimTime() const;
+
+	FName GetCurrentSection() const;
+
+	int32 GetCurrentSectionIndex() const;
+
+	FAnimMontageInstance* GetAnimMontageInstance() const;
+
+	const UAnimMontage* GetAnimMontage() const;
+
+	UAnimInstance* GetAnimInstance() const;
+
+	USkeletalMeshComponent* GetSkeletalMeshComponent() const;
+
+	UContextualAnimSceneActorComponent* GetSceneActorComponent() const;
+
+private:
+
+	friend UContextualAnimSceneInstance;
+
+	/** The actual actor in the world */
+	TWeakObjectPtr<AActor> Actor = nullptr;
+
+	/** Ptr to the animation data in the scene asset used by this actor */
+	const FContextualAnimData* AnimDataPtr = nullptr;
+
+	const FContextualAnimTrackSettings* SettingsPtr = nullptr;
+
+	/** Desired time to start the animation */
+	float AnimStartTime = 0.f;
+
+	/** Ptr back to the scene instance we belong to */
+	TWeakObjectPtr<const UContextualAnimSceneInstance> SceneInstancePtr = nullptr;
 };
 
 /** Delegate to notify external objects when this is scene is completed */
@@ -54,7 +94,7 @@ DECLARE_DELEGATE_TwoParams(FOnContextualAnimSceneActorJoined, class UContextualA
 DECLARE_DELEGATE_TwoParams(FOnContextualAnimSceneActorLeft, class UContextualAnimSceneInstance*, AActor* Actor)
 
 /** Instance of a contextual animation scene */
-UCLASS()
+UCLASS(BlueprintType, Blueprintable)
 class CONTEXTUALANIMATION_API UContextualAnimSceneInstance : public UObject
 {
 	GENERATED_BODY()
@@ -96,17 +136,40 @@ public:
 	/** Whether the supplied actor is part of this scene */
 	bool IsActorInThisScene(const AActor* Actor) const;
 
+	const FContextualAnimSceneActorData* FindSceneActorDataForActor(const AActor* Actor) const;
+
+	const FContextualAnimSceneActorData* FindSceneActorDataForRole(const FName& Role) const;
+
+	UFUNCTION(BlueprintCallable, Category = "Contextual Anim|Scene Instance")
+	float GetCurrentSectionTimeLeft() const;
+
+	UFUNCTION(BlueprintCallable, Category = "Contextual Anim|Scene Instance")
+	bool DidCurrentSectionLoop() const;
+
+	UFUNCTION(BlueprintCallable, Category = "Contextual Anim|Scene Instance")
+	float GetPositionInCurrentSection() const;
+
+	UFUNCTION(BlueprintNativeEvent, Category = "Contextual Anim|Scene Instance")
+	float GetResumePositionForSceneActor(const FContextualAnimSceneActorData& SceneActorData, int32 DesiredSectionIndex) const;
+
+	UFUNCTION(BlueprintCallable, Category = "Contextual Anim|Scene Instance")
+	AActor* GetActorByRole(FName Role) const;
+
 protected:
 
 	/** Tells the scene actor to join the scene (play animation) */
-	void Join(FContextualAnimSceneActorData& Data);
+	void Join(FContextualAnimSceneActorData& SceneActorData);
 
 	/** Tells the scene actor to leave the scene (stop animation) */
-	void Leave(FContextualAnimSceneActorData& Data);
+	void Leave(FContextualAnimSceneActorData& SceneActorData);
+
+	bool TransitionTo(FContextualAnimSceneActorData& SceneActorData, const FName& ToSectionName);
 
 	/** Helper function to set ignore collision between the supplied actor and all the other actors in this scene */
 	void SetIgnoreCollisionWithOtherActors(AActor* Actor, bool bValue) const;
 
+	void UpdateTransitions(float DeltaTime);
+
 	UFUNCTION()
 	void OnMontageBlendingOut(UAnimMontage* Montage, bool bInterrupted);
 
@@ -115,4 +178,10 @@ protected:
 
 	UFUNCTION()
 	void OnNotifyEndReceived(FName NotifyName, const FBranchingPointNotifyPayload& BranchingPointNotifyPayload);
+
+public:
+
+	/** Extracts data from a ContextualAnimSceneActorData */
+	UFUNCTION(BlueprintPure, Category = "Contextual Anim|Scene Actor Data", meta = (NativeBreakFunc))
+	static void BreakContextualAnimSceneActorData(const FContextualAnimSceneActorData& SceneActorData, AActor*& Actor, UAnimMontage*& Montage, float& AnimTime, int32& CurrentSectionIndex, FName& CurrentSectionName);
 };
\ No newline at end of file
diff --git a/Engine/Plugins/Experimental/Animation/ContextualAnimation/Source/ContextualAnimation/Public/ContextualAnimScenePivotProvider.h b/Engine/Plugins/Experimental/Animation/ContextualAnimation/Source/ContextualAnimation/Public/ContextualAnimScenePivotProvider.h
--- a/Engine/Plugins/Experimental/Animation/ContextualAnimation/Source/ContextualAnimation/Public/ContextualAnimScenePivotProvider.h
+++ b/Engine/Plugins/Experimental/Animation/ContextualAnimation/Source/ContextualAnimation/Public/ContextualAnimScenePivotProvider.h
@@ -18,7 +18,7 @@ public:
 
 	UContextualAnimScenePivotProvider(const FObjectInitializer& ObjectInitializer);
 
-	virtual FTransform CalculateScenePivot_Source() const { return FTransform::Identity; }
+	virtual FTransform CalculateScenePivot_Source(int32 AnimDataIndex) const { return FTransform::Identity; }
 	virtual FTransform CalculateScenePivot_Runtime(const TMap<FName, FContextualAnimSceneActorData>& SceneActorMap) const { return FTransform::Identity; }
 
 	const UContextualAnimSceneAssetBase* GetSceneAsset() const;
@@ -42,7 +42,7 @@ public:
 
 	UContextualAnimScenePivotProvider_Default(const FObjectInitializer& ObjectInitializer);
 
-	virtual FTransform CalculateScenePivot_Source() const override;
+	virtual FTransform CalculateScenePivot_Source(int32 AnimDataIndex) const override;
 	virtual FTransform CalculateScenePivot_Runtime(const TMap<FName, FContextualAnimSceneActorData>& SceneActorMap) const override;
 };
 
@@ -58,6 +58,6 @@ public:
 
 	UContextualAnimScenePivotProvider_RelativeTo(const FObjectInitializer& ObjectInitializer);
 
-	virtual FTransform CalculateScenePivot_Source() const override;
+	virtual FTransform CalculateScenePivot_Source(int32 AnimDataIndex) const override;
 	virtual FTransform CalculateScenePivot_Runtime(const TMap<FName, FContextualAnimSceneActorData>& SceneActorMap) const override;
 };
\ No newline at end of file
diff --git a/Engine/Plugins/Experimental/Animation/ContextualAnimation/Source/ContextualAnimation/Public/ContextualAnimTransition.h b/Engine/Plugins/Experimental/Animation/ContextualAnimation/Source/ContextualAnimation/Public/ContextualAnimTransition.h
new file mode 100644
--- /dev/null
+++ b/Engine/Plugins/Experimental/Animation/ContextualAnimation/Source/ContextualAnimation/Public/ContextualAnimTransition.h
@@ -0,0 +1,22 @@
+// Copyright Epic Games, Inc. All Rights Reserved.
+
+#pragma once
+
+#include "UObject/Object.h"
+#include "ContextualAnimTransition.generated.h"
+
+class AActor;
+class UContextualAnimSceneInstance;
+
+UCLASS(Const, Abstract, BlueprintType, Blueprintable, EditInlineNew)
+class CONTEXTUALANIMATION_API UContextualAnimTransition : public UObject
+{
+	GENERATED_BODY()
+
+public:
+
+	UContextualAnimTransition(const FObjectInitializer& ObjectInitializer);
+
+	UFUNCTION(BlueprintImplementableEvent, Category = "Defaults")
+	bool CanEnterTransition(const UContextualAnimSceneInstance* SceneInstance, const FName& FromSection, const FName& ToSection) const;
+};
\ No newline at end of file
diff --git a/Engine/Plugins/Experimental/Animation/ContextualAnimation/Source/ContextualAnimation/Public/ContextualAnimTypes.h b/Engine/Plugins/Experimental/Animation/ContextualAnimation/Source/ContextualAnimation/Public/ContextualAnimTypes.h
--- a/Engine/Plugins/Experimental/Animation/ContextualAnimation/Source/ContextualAnimation/Public/ContextualAnimTypes.h
+++ b/Engine/Plugins/Experimental/Animation/ContextualAnimation/Source/ContextualAnimation/Public/ContextualAnimTypes.h
@@ -6,12 +6,19 @@
 #include "Templates/SubclassOf.h"
 #include "ContextualAnimTypes.generated.h"
 
-DECLARE_LOG_CATEGORY_EXTERN(LogContextualAnim, Log, All);
+CONTEXTUALANIMATION_API DECLARE_LOG_CATEGORY_EXTERN(LogContextualAnim, Log, All);
 
 class AActor;
 class UAnimMontage;
 class UContextualAnimMetadata;
 
+#if !(UE_BUILD_SHIPPING || UE_BUILD_TEST)
+struct FContextualAnimCVars
+{
+	static TAutoConsoleVariable<int32> CVarDrawDebugIKTargets;
+};
+#endif
+
 /** Container for alignment tracks */
 USTRUCT()
 struct CONTEXTUALANIMATION_API FContextualAnimAlignmentTrackContainer
@@ -24,10 +31,35 @@ struct CONTEXTUALANIMATION_API FContextualAnimAlignmentTrackContainer
 	UPROPERTY()
 	float SampleInterval = 0.f;
 
+	void Initialize(int32 InNumTracks, float InSampleInterval)
+	{
+		Tracks.AnimationTracks.Empty(InNumTracks);
+		Tracks.TrackNames.Empty(InNumTracks);
+		SampleInterval = InSampleInterval;
+	}
+
+	void Empty()
+	{
+		Tracks.AnimationTracks.Empty();
+		Tracks.TrackNames.Empty();
+		SampleInterval = 0.f;
+	}
+
 	FTransform ExtractTransformAtTime(int32 TrackIndex, float Time) const;
 	FTransform ExtractTransformAtTime(const FName& TrackName, float Time) const;
 };
 
+USTRUCT()
+struct CONTEXTUALANIMATION_API FContextualAnimRoleBonePair
+{
+	GENERATED_BODY()
+
+	UPROPERTY()
+	FName RoleName = NAME_None;
+
+	UPROPERTY()
+	FName BoneName = NAME_None;
+};
 
 USTRUCT(BlueprintType)
 struct CONTEXTUALANIMATION_API FContextualAnimData
@@ -40,18 +72,30 @@ struct CONTEXTUALANIMATION_API FContextualAnimData
 	UPROPERTY(EditAnywhere, BlueprintReadWrite, Category = "Defaults")
 	float AnimMaxStartTime = 0.f;
 
+	UPROPERTY(EditAnywhere, BlueprintReadWrite, Category = "Defaults")
+	bool bRequireFlyingMode = false;
+
+	/** Container for alignment tracks */
 	UPROPERTY()
 	FContextualAnimAlignmentTrackContainer AlignmentData;
 
+	/** Container for auto generate IK Target Tracks */
+	UPROPERTY()
+	FContextualAnimAlignmentTrackContainer IKTargetData;
+
+	/** Lookup map to go from TrackName to Target Role and Parent Bone */
+	UPROPERTY()
+	TMap<FName, FContextualAnimRoleBonePair> IKTargetTrackLookupMap;
+
 	UPROPERTY(EditAnywhere, Instanced, Category = "Defaults")
 	UContextualAnimMetadata* Metadata = nullptr;
 
 	UPROPERTY(EditAnywhere, BlueprintReadWrite, Category = "Defaults")
 	FTransform MeshToScene;
 
-	/** deprecated */
+	/** Index in the container of AnimData that owns us */
 	UPROPERTY()
-	float SyncTime = 0.f;
+	int32 Index = INDEX_NONE;
 
 	float GetSyncTimeForWarpSection(int32 WarpSectionIndex) const;
 	float GetSyncTimeForWarpSection(const FName& WarpSectionName) const;
@@ -73,6 +117,94 @@ enum class EContextualAnimJoinRule : uint8
 	Late
 };
 
+USTRUCT(BlueprintType)
+struct FContextualAnimTransitionContainer
+{
+	GENERATED_BODY()
+
+	UPROPERTY(EditAnywhere, BlueprintReadWrite, Category = "Defaults")
+	TArray<FName> FromSections;
+
+	UPROPERTY(EditAnywhere, BlueprintReadWrite, Category = "Defaults")
+	FName ToSection = NAME_None;
+
+	UPROPERTY(EditAnywhere, Instanced, Category = "Defaults")
+	class UContextualAnimTransition* Transition = nullptr;
+
+	//@TODO: Remove this from here it should be on the EdMode
+	UPROPERTY(EditAnywhere, BlueprintReadWrite, Category = "Defaults")
+	bool bForceTransition = false;
+};
+
+UENUM(BlueprintType)
+enum class EContextualAnimIKTargetProvider : uint8
+{
+	/** 
+	 * IK Targets are auto generated from the animations 
+	 * @see: UContextualAnimSceneAssetBase::GenerateIKTargetTracks
+	 */
+	Autogenerated,
+
+	/** IK Target defined by a bone or a socket in the animation */
+	Bone,
+
+	//@TODO: Add collision query methods
+};
+
+USTRUCT(BlueprintType)
+struct FContextualAnimIKTargetBoneParams
+{
+	GENERATED_BODY()
+
+	UPROPERTY(EditAnywhere, BlueprintReadWrite, Category = "Defaults")
+	FName TargetRole = NAME_None;
+
+	UPROPERTY(EditAnywhere, BlueprintReadWrite, Category = "Defaults")
+	FName BoneName = NAME_None;
+};
+
+USTRUCT(BlueprintType)
+struct FContextualAnimIKTargetDefinition
+{
+	GENERATED_BODY()
+
+	UPROPERTY(EditAnywhere, BlueprintReadWrite, Category = "Defaults")
+	FName IKGoalName = NAME_None;
+
+	UPROPERTY(EditAnywhere, BlueprintReadWrite, Category = "Defaults")
+	FName AlphaCurveName = NAME_None;
+
+	UPROPERTY(EditAnywhere, BlueprintReadWrite, Category = "Defaults")
+	EContextualAnimIKTargetProvider Provider;
+
+	//@TODO: Hide from the UI when Type != Auto
+	UPROPERTY(EditAnywhere, BlueprintReadWrite, Category = "Defaults", meta = (EditCondition = "Provider == EContextualAnimIKTargetProvider::Autogenerated"))
+	FContextualAnimIKTargetBoneParams AutoParams;
+
+	//@TODO: Hide from the UI when Type != Bone
+	UPROPERTY(EditAnywhere, BlueprintReadWrite, Category = "Defaults", meta = (EditCondition = "Provider == EContextualAnimIKTargetProvider::Bone"))
+	FContextualAnimIKTargetBoneParams BoneParams;
+};
+
+USTRUCT(BlueprintType)
+struct CONTEXTUALANIMATION_API FContextualAnimIKTarget
+{
+	GENERATED_BODY()
+
+	UPROPERTY(EditAnywhere, BlueprintReadWrite, Category = "Defaults")
+	FName GoalName = NAME_None;
+
+	UPROPERTY(EditAnywhere, BlueprintReadWrite, Category = "Defaults")
+	float Alpha = 0.f;
+
+	UPROPERTY(EditAnywhere, BlueprintReadWrite, Category = "Defaults")
+	FTransform Transform;
+
+	FContextualAnimIKTarget(){}
+	FContextualAnimIKTarget(const FName& InGoalName, float InAlpha, const FTransform& InTransform)
+		: GoalName(InGoalName), Alpha(InAlpha), Transform(InTransform){}
+};
+
 USTRUCT(BlueprintType)
 struct CONTEXTUALANIMATION_API FContextualAnimTrackSettings
 {
@@ -83,6 +215,12 @@ struct CONTEXTUALANIMATION_API FContextualAnimTrackSettings
 
 	UPROPERTY(EditAnywhere, BlueprintReadWrite, Category = "Defaults")
 	EContextualAnimJoinRule JoinRule = EContextualAnimJoinRule::Default;
+
+	UPROPERTY(EditAnywhere, BlueprintReadWrite, Category = "Defaults")
+	TArray<FContextualAnimIKTargetDefinition> IKTargetDefinitions;
+
+	UPROPERTY(EditAnywhere, BlueprintReadWrite, Category = "Defaults", AdvancedDisplay)
+	FTransform MeshToComponent = FTransform(FRotator(0.f, -90.f, 0.f));
 };
 
 USTRUCT(BlueprintType)
@@ -107,4 +245,117 @@ struct FContextualAnimCompositeTrack
 
 	UPROPERTY(EditAnywhere, BlueprintReadWrite, Category = "Defaults")
 	TArray<FContextualAnimData> AnimDataContainer;
+
+	FTransform GetRootTransformForAnimDataAtIndex(int32 Index) const;
+};
+
+USTRUCT(BlueprintType)
+struct FContextualAnimSceneBindings
+{
+	GENERATED_BODY()
+
+	/** Map with actors to bind to each role in the scene */
+	UPROPERTY(EditAnywhere, BlueprintReadWrite, Category = "Defaults")
+	TMap<FName, AActor*> RoleToActorMap;
+
+	/** Desired AnimDataIndex. If INDEX_NONE the Manager will attempt to find the best AnimData to use */
+	UPROPERTY(EditAnywhere, BlueprintReadWrite, Category = "Defaults")
+	int32 AnimDataIndex = INDEX_NONE;
+
+	/** Desired start time. Only relevant if AnimDataIndex != INDEX_NONE */
+	UPROPERTY(EditAnywhere, BlueprintReadWrite, Category = "Defaults")
+	float AnimStartTime = 0.f;
+};
+
+///////////////////////////////////////////////////////////////////////
+
+/** Stores the result of a query function */
+USTRUCT(BlueprintType)
+struct FContextualAnimQueryResult
+{
+	GENERATED_BODY()
+
+	UPROPERTY(BlueprintReadWrite, Category = "Defaults")
+	TWeakObjectPtr<UAnimMontage> Animation;
+
+	UPROPERTY(BlueprintReadWrite, Category = "Defaults")
+	FTransform EntryTransform;
+
+	UPROPERTY(BlueprintReadWrite, Category = "Defaults")
+	FTransform SyncTransform;
+
+	UPROPERTY(BlueprintReadWrite, Category = "Defaults")
+	float AnimStartTime = 0.f;
+
+	UPROPERTY(BlueprintReadWrite, Category = "Defaults")
+	int32 DataIndex = INDEX_NONE;
+
+	void Reset()
+	{
+		Animation.Reset();
+		EntryTransform = SyncTransform = FTransform::Identity;
+		AnimStartTime = 0.f;
+		DataIndex = INDEX_NONE;
+	}
+
+	FORCEINLINE bool IsValid() const { return DataIndex != INDEX_NONE; }
+};
+
+/** Stores the parameters passed into query function */
+USTRUCT(BlueprintType)
+struct FContextualAnimQueryParams
+{
+	GENERATED_BODY()
+
+	UPROPERTY(EditAnywhere, BlueprintReadWrite, Category = "Defaults")
+	TWeakObjectPtr<const AActor> Querier;
+
+	UPROPERTY(EditAnywhere, BlueprintReadWrite, Category = "Defaults")
+	FTransform QueryTransform;
+
+	UPROPERTY(EditAnywhere, BlueprintReadWrite, Category = "Defaults")
+	bool bComplexQuery = false;
+
+	UPROPERTY(EditAnywhere, BlueprintReadWrite, Category = "Defaults")
+	bool bFindAnimStartTime = false;
+
+	FContextualAnimQueryParams() {}
+
+	FContextualAnimQueryParams(const AActor* InQuerier, bool bInComplexQuery, bool bInFindAnimStartTime)
+		: Querier(InQuerier), bComplexQuery(bInComplexQuery), bFindAnimStartTime(bInFindAnimStartTime) {}
+
+	FContextualAnimQueryParams(const FTransform& InQueryTransform, bool bInComplexQuery, bool bInFindAnimStartTime)
+		: QueryTransform(InQueryTransform), bComplexQuery(bInComplexQuery), bFindAnimStartTime(bInFindAnimStartTime) {}
+};
+
+USTRUCT(BlueprintType)
+struct FContextualAnimQuerier
+{
+	GENERATED_BODY()
+
+	UPROPERTY(EditAnywhere, BlueprintReadWrite, Category = "Defaults")
+	TWeakObjectPtr<const AActor> Actor = nullptr;
+
+	UPROPERTY(EditAnywhere, BlueprintReadWrite, Category = "Defaults")
+	FTransform Transform;
+
+	FContextualAnimQuerier(){}
+	FContextualAnimQuerier(const AActor* InActor) : Actor(InActor) {}
+	FContextualAnimQuerier(const FTransform& InTransform) : Transform(InTransform) {}
+};
+
+USTRUCT(BlueprintType)
+struct FContextualAnimQueryContext
+{
+	GENERATED_BODY()
+
+	UPROPERTY(EditAnywhere, BlueprintReadWrite, Category = "Defaults")
+	TWeakObjectPtr<const AActor> Actor = nullptr;
+
+	UPROPERTY(EditAnywhere, BlueprintReadWrite, Category = "Defaults")
+	FTransform Transform;
+
+	FContextualAnimQueryContext() {}
+	FContextualAnimQueryContext(const AActor* InActor) : Actor(InActor) {}
+	FContextualAnimQueryContext(const FTransform& InTransform) : Transform(InTransform) {}
 };
\ No newline at end of file
diff --git a/Engine/Plugins/Experimental/Animation/ContextualAnimation/Source/ContextualAnimation/Public/ContextualAnimUtilities.h b/Engine/Plugins/Experimental/Animation/ContextualAnimation/Source/ContextualAnimation/Public/ContextualAnimUtilities.h
--- a/Engine/Plugins/Experimental/Animation/ContextualAnimation/Source/ContextualAnimation/Public/ContextualAnimUtilities.h
+++ b/Engine/Plugins/Experimental/Animation/ContextualAnimation/Source/ContextualAnimation/Public/ContextualAnimUtilities.h
@@ -7,6 +7,9 @@
 #include "ContextualAnimUtilities.generated.h"
 
 class UContextualAnimSceneAsset;
+class USkeletalMeshComponent;
+class UAnimInstance;
+class AActor;
 
 UCLASS()
 class CONTEXTUALANIMATION_API UContextualAnimUtilities : public UBlueprintFunctionLibrary
@@ -29,8 +32,26 @@ public:
 	 */
 	static void ExtractComponentSpacePose(const UAnimSequenceBase* Animation, const FBoneContainer& BoneContainer, float Time, bool bExtractRootMotion, FCSPose<FCompactPose>& OutPose);
 
+	/** Extract Root Motion transform from a contiguous position range */
 	static FTransform ExtractRootMotionFromAnimation(const UAnimSequenceBase* Animation, float StartTime, float EndTime);
 
+	/** Extract root bone transform at a given time */
+	static FTransform ExtractRootTransformFromAnimation(const UAnimSequenceBase* Animation, float Time);
+
 	static void DrawDebugPose(const UWorld* World, const UAnimSequenceBase* Animation, float Time, const FTransform& LocalToWorldTransform, const FColor& Color, float LifeTime, float Thickness);
-	static void DrawDebugScene(const UWorld* World, const UContextualAnimSceneAsset* SceneAsset, float Time, const FTransform& ToWorldTransform, const FColor& Color, float LifeTime, float Thickness);
+	
+	static void DrawDebugScene(const UWorld* World, const UContextualAnimSceneAsset* SceneAsset, int32 AnimDataIndex, float Time, const FTransform& ToWorldTransform, const FColor& Color, float LifeTime, float Thickness);
+
+	static USkeletalMeshComponent* TryGetSkeletalMeshComponent(AActor* Actor);
+
+	static UAnimInstance* TryGetAnimInstance(AActor* Actor);
+
+	UFUNCTION(BlueprintCallable, Category = "Contextual Anim|Utilities", meta = (DisplayName = "GetSectionStartAndEndTime"))
+	static void BP_Montage_GetSectionStartAndEndTime(const UAnimMontage* Montage, int32 SectionIndex, float& OutStartTime, float& OutEndTime);
+
+	UFUNCTION(BlueprintCallable, Category = "Contextual Anim|Utilities", meta = (DisplayName = "GetSectionTimeLeftFromPos"))
+	static float BP_Montage_GetSectionTimeLeftFromPos(const UAnimMontage* Montage, float Position);
+
+	UFUNCTION(BlueprintCallable, Category = "Contextual Anim|Utilities", meta = (DisplayName = "GetSectionLength"))
+	static float BP_Montage_GetSectionLength(const UAnimMontage* Montage, int32 SectionIndex);
 };
\ No newline at end of file
diff --git a/Engine/Plugins/Experimental/Animation/ContextualAnimation/Source/ContextualAnimationEditor/Private/ContextualAnimEdMode.cpp b/Engine/Plugins/Experimental/Animation/ContextualAnimation/Source/ContextualAnimationEditor/Private/ContextualAnimEdMode.cpp
--- a/Engine/Plugins/Experimental/Animation/ContextualAnimation/Source/ContextualAnimationEditor/Private/ContextualAnimEdMode.cpp
+++ b/Engine/Plugins/Experimental/Animation/ContextualAnimation/Source/ContextualAnimationEditor/Private/ContextualAnimEdMode.cpp
@@ -17,6 +17,7 @@
 #include "ContextualAnimSceneActorComponent.h"
 #include "GameFramework/Character.h"
 #include "ContextualAnimMetadata.h"
+#include "ContextualAnimSceneInstance.h"
 
 const FEditorModeID FContextualAnimEdMode::EM_ContextualAnimEdModeId = TEXT("EM_ContextualAnimEdMode");
 
@@ -91,7 +92,7 @@ void FContextualAnimEdMode::Tick(FEditorViewportClient* ViewportClient, float De
 				if (UContextualAnimSceneAsset* Asset = GetContextualAnimEdModeToolkit()->GetSettings()->SceneAsset)
 				{
 					DrawDebugCoordinateSystem(GetWorld(), PreviewManager->ScenePivot.GetLocation(), PreviewManager->ScenePivot.Rotator(), 50.f, false, 0.f, 0, 1.f);
-					UContextualAnimUtilities::DrawDebugScene(GetWorld(), Asset, PreviewManager->Time, PreviewManager->ScenePivot, FColor::White, 0.f, 1.f);
+					UContextualAnimUtilities::DrawDebugScene(GetWorld(), Asset, PreviewManager->AnimDataIndex, PreviewManager->Time, PreviewManager->ScenePivot, FColor::White, 0.f, 1.f);
 				}
 			}
 
@@ -116,35 +117,20 @@ bool FContextualAnimEdMode::InputKey(FEditorViewportClient* ViewportClient, FVie
 			UContextualAnimSceneAsset* Asset = GetContextualAnimEdModeToolkit()->GetSettings()->SceneAsset;
 			if (Asset && PreviewManager->PreviewActors.Num() > 0)
 			{
-				Manager->TryStartScene(Asset, PreviewManager->PreviewActors);
+				FContextualAnimSceneBindings Bindings;
+				Bindings.RoleToActorMap = PreviewManager->PreviewActors;
+				Manager->TryStartScene(Asset, Bindings);
 			}
 			else
 			{
-				UContextualAnimSceneActorComponent* Comp = nullptr;
-
-				TArray<UPrimitiveComponent*> OverlappingComps;
-				PreviewManager->TestCharacter->GetOverlappingComponents(OverlappingComps);
-
-				for (UPrimitiveComponent* OverlappingComp : OverlappingComps)
-				{
-					if (OverlappingComp && OverlappingComp->GetClass()->IsChildOf<UContextualAnimSceneActorComponent>())
-					{
-						Comp = Cast<UContextualAnimSceneActorComponent>(OverlappingComp);
-						break;
-					}
-				}
-
-				if (Comp)
-				{
-					TMap<FName, AActor*> Bindings;
-					Bindings.Add(UContextualAnimCompositeSceneAsset::InteractableRoleName, Comp->GetOwner());
-					Bindings.Add(UContextualAnimCompositeSceneAsset::InteractorRoleName, PreviewManager->TestCharacter.Get());
-
-					UContextualAnimManager::Get(GetWorld())->TryStartScene(Comp->SceneAsset, Bindings);
+				if (UContextualAnimSceneActorComponent* Comp = Manager->FindClosestSceneActorCompToActor(PreviewManager->TestCharacter.Get()))
+				{	
+					Manager->TryStartScene(Comp->SceneAsset, Comp->GetOwner(), Manager->GetSceneActorCompContainer());
 				}
 				else
 				{
-					GEngine->AddOnScreenDebugMessage(-1, 5.f, FColor::Yellow, FString(TEXT("WARNING: The preview actor is not overlapping any interactable")));
+					GEngine->AddOnScreenDebugMessage(-1, 5.f, FColor::Yellow, FString(TEXT("WARNING: Can't find any SceneActor to interact with")));
+					UE_LOG(LogContextualAnim, Warning, TEXT("WARNING: Can't find any SceneActor to interact with"));
 				}
 			}
 		}
@@ -169,18 +155,25 @@ bool FContextualAnimEdMode::HandleClick(FEditorViewportClient* InViewportClient,
 	if (!GEditor->IsSimulatingInEditor())
 	{
 		GEngine->AddOnScreenDebugMessage(-1, 5.f, FColor::Yellow, FString(TEXT("WARNING. You are not in Simulating Mode")));
+		UE_LOG(LogContextualAnim, Warning, TEXT("WARNING. You are not in Simulating Mode"));
+
 		return FEdMode::HandleClick(InViewportClient, HitProxy, Click);
 	}
 
 	FHitResult HitResult;
 	GetHitResultUnderCursor(HitResult, InViewportClient, Click);
 
-	GEngine->AddOnScreenDebugMessage(-1, 5.f, FColor::White, FString::Printf(TEXT("INFO: HandleClick: IsValidHit: %d Actor: %s Location: %s"),
-		HitResult.IsValidBlockingHit(), *GetNameSafe(HitResult.GetActor()), *HitResult.ImpactPoint.ToString()));
+	const FString Message = FString::Printf(TEXT("INFO: HandleClick: IsValidHit: %d Actor: %s Location: %s"),
+		HitResult.IsValidBlockingHit(), *GetNameSafe(HitResult.GetActor()), *HitResult.ImpactPoint.ToString());
+
+	GEngine->AddOnScreenDebugMessage(-1, 5.f, FColor::White, Message);
+	UE_LOG(LogContextualAnim, Warning, TEXT("%s"), *Message);
 
 	if (!HitResult.IsValidBlockingHit())
 	{
 		GEngine->AddOnScreenDebugMessage(-1, 5.f, FColor::Yellow, FString(TEXT("WARNING. HitResult from click event is not a valid blocking hit.")));
+		UE_LOG(LogContextualAnim, Warning, TEXT("WARNING. HitResult from click event is not a valid blocking hit."));
+
 		return FEdMode::HandleClick(InViewportClient, HitProxy, Click);
 	}	
 	
diff --git a/Engine/Plugins/Experimental/Animation/ContextualAnimation/Source/ContextualAnimationEditor/Private/ContextualAnimPreviewManager.cpp b/Engine/Plugins/Experimental/Animation/ContextualAnimation/Source/ContextualAnimationEditor/Private/ContextualAnimPreviewManager.cpp
--- a/Engine/Plugins/Experimental/Animation/ContextualAnimation/Source/ContextualAnimationEditor/Private/ContextualAnimPreviewManager.cpp
+++ b/Engine/Plugins/Experimental/Animation/ContextualAnimation/Source/ContextualAnimationEditor/Private/ContextualAnimPreviewManager.cpp
@@ -56,25 +56,28 @@ void UContextualAnimPreviewManager::SpawnPreviewActors(const UContextualAnimScen
 
 		for (const auto& Entry : SceneAsset->DataContainer)
 		{
-			const FName& Role = Entry.Key;
-			const FContextualAnimData& Data = Entry.Value.AnimData;
-
-			const FTransform SpawnTransform = (Data.AlignmentData.ExtractTransformAtTime(0, 0.f) * SceneOrigin);
-
-			UClass* PreviewClass = SceneAsset->GetPreviewActorClassForRole(Role);
-			if(!PreviewClass && DefaultPreviewClass)
+			if(Entry.Value.AnimDataContainer.Num() > 0)
 			{
-				PreviewClass = DefaultPreviewClass.Get();
-			}
+				const FName& Role = Entry.Key;
+				const FContextualAnimData& Data = Entry.Value.AnimDataContainer[0];
 
-			AActor* PreviewActor = SpawnPreviewActor(PreviewClass, SpawnTransform);
-			if (PreviewActor)
-			{
-				PreviewActors.Add(Role, PreviewActor);
+				const FTransform SpawnTransform = (Data.AlignmentData.ExtractTransformAtTime(0, 0.f) * SceneOrigin);
 
-				if (!TestCharacter.IsValid())
+				UClass* PreviewClass = SceneAsset->GetTrackSettings(Role)->PreviewActorClass;
+				if (!PreviewClass && DefaultPreviewClass)
 				{
-					TestCharacter = Cast<ACharacter>(PreviewActor);
+					PreviewClass = DefaultPreviewClass.Get();
+				}
+
+				AActor* PreviewActor = SpawnPreviewActor(PreviewClass, SpawnTransform);
+				if (PreviewActor)
+				{
+					PreviewActors.Add(Role, PreviewActor);
+
+					if (!TestCharacter.IsValid())
+					{
+						TestCharacter = Cast<ACharacter>(PreviewActor);
+					}
 				}
 			}
 		}
diff --git a/Engine/Plugins/Experimental/Animation/ContextualAnimation/Source/ContextualAnimationEditor/Private/ContextualAnimationEditor.cpp b/Engine/Plugins/Experimental/Animation/ContextualAnimation/Source/ContextualAnimationEditor/Private/ContextualAnimationEditor.cpp
--- a/Engine/Plugins/Experimental/Animation/ContextualAnimation/Source/ContextualAnimationEditor/Private/ContextualAnimationEditor.cpp
+++ b/Engine/Plugins/Experimental/Animation/ContextualAnimation/Source/ContextualAnimationEditor/Private/ContextualAnimationEditor.cpp
@@ -3,7 +3,6 @@
 #include "ContextualAnimationEditor.h"
 #include "ContextualAnimEdMode.h"
 #include "ContextualAnimEditorStyle.h"
-#include "Customizations/MathStructCustomizations.h"
 
 #define LOCTEXT_NAMESPACE "FContextualAnimationEditorModule"
 
@@ -12,10 +11,6 @@ void FContextualAnimationEditorModule::StartupModule()
 	FContextualAnimEditorStyle::Initialize();
 	FContextualAnimEditorStyle::ReloadTextures();
 
-	FPropertyEditorModule& PropertyModule = FModuleManager::LoadModuleChecked<FPropertyEditorModule>("PropertyEditor");
-	PropertyModule.RegisterCustomPropertyTypeLayout("ContextualAnimDistanceParam", FOnGetPropertyTypeCustomizationInstance::CreateStatic(&FMathStructCustomization::MakeInstance));
-	PropertyModule.RegisterCustomPropertyTypeLayout("ContextualAnimAngleParam", FOnGetPropertyTypeCustomizationInstance::CreateStatic(&FMathStructCustomization::MakeInstance));
-
 	FEditorModeRegistry::Get().RegisterMode<FContextualAnimEdMode>(FContextualAnimEdMode::EM_ContextualAnimEdModeId, 
 		LOCTEXT("ContextualAnimEdModeEdModeName", "ContextualAnim"), 
 		FSlateIcon(FContextualAnimEditorStyle::GetStyleSetName(), "ContextualAnimEditor.Icon", "ContextualAnimEditor.Icon"),
@@ -26,10 +21,6 @@ void FContextualAnimationEditorModule::ShutdownModule()
 {
 	FContextualAnimEditorStyle::Shutdown();
 
-	FPropertyEditorModule& PropertyModule = FModuleManager::LoadModuleChecked<FPropertyEditorModule>("PropertyEditor");
-	PropertyModule.UnregisterCustomPropertyTypeLayout("ContextualAnimDistanceParam");
-	PropertyModule.UnregisterCustomPropertyTypeLayout("ContextualAnimAngleParam");
-
 	FEditorModeRegistry::Get().UnregisterMode(FContextualAnimEdMode::EM_ContextualAnimEdModeId);
 }
 
diff --git a/Engine/Plugins/Experimental/Animation/ContextualAnimation/Source/ContextualAnimationEditor/Public/ContextualAnimPreviewManager.h b/Engine/Plugins/Experimental/Animation/ContextualAnimation/Source/ContextualAnimationEditor/Public/ContextualAnimPreviewManager.h
--- a/Engine/Plugins/Experimental/Animation/ContextualAnimation/Source/ContextualAnimationEditor/Public/ContextualAnimPreviewManager.h
+++ b/Engine/Plugins/Experimental/Animation/ContextualAnimation/Source/ContextualAnimationEditor/Public/ContextualAnimPreviewManager.h
@@ -23,6 +23,9 @@ public:
 	UPROPERTY(EditAnywhere, Category = "Debug")
 	bool bDrawDebugScene;
 
+	UPROPERTY(EditAnywhere, Category = "Debug", meta = (EditCondition = "bDrawDebugScene"))
+	int32 AnimDataIndex = 0;
+
 	UPROPERTY(EditAnywhere, Category = "Debug", meta = (EditCondition = "bDrawDebugScene"))
 	float Time;
 
