ContextualAnim - Removed CompositeSceneAsset and SceneAssetBase #jira none #rb [REDACTED_HANDLE]

diff --git a/Engine/Plugins/Experimental/Animation/ContextualAnimation/Source/ContextualAnimation/Private/ContextualAnimCompositeSceneAsset.cpp b/Engine/Plugins/Experimental/Animation/ContextualAnimation/Source/ContextualAnimation/Private/ContextualAnimCompositeSceneAsset.cpp
deleted file mode 100644
diff --git a/Engine/Plugins/Experimental/Animation/ContextualAnimation/Source/ContextualAnimation/Private/ContextualAnimManager.cpp b/Engine/Plugins/Experimental/Animation/ContextualAnimation/Source/ContextualAnimation/Private/ContextualAnimManager.cpp
--- a/Engine/Plugins/Experimental/Animation/ContextualAnimation/Source/ContextualAnimation/Private/ContextualAnimManager.cpp
+++ b/Engine/Plugins/Experimental/Animation/ContextualAnimation/Source/ContextualAnimation/Private/ContextualAnimManager.cpp
@@ -5,7 +5,6 @@
 #include "ContextualAnimSceneInstance.h"
 #include "DrawDebugHelpers.h"
 #include "ContextualAnimSceneAsset.h"
-#include "ContextualAnimCompositeSceneAsset.h"
 #include "ContextualAnimation.h"
 #include "ContextualAnimSceneActorComponent.h"
 #include "Engine/World.h"
@@ -131,7 +130,7 @@ UContextualAnimSceneInstance* UContextualAnimManager::GetSceneWithActor(AActor*
 	return nullptr;
 }
 
-bool UContextualAnimManager::TryStartScene(const UContextualAnimSceneAssetBase* SceneAsset, const FContextualAnimSceneBindings& Bindings)
+bool UContextualAnimManager::TryStartScene(const UContextualAnimSceneAsset* SceneAsset, const FContextualAnimSceneBindings& Bindings)
 {
 	if (SceneAsset == nullptr)
 	{
@@ -232,7 +231,7 @@ bool UContextualAnimManager::TryStartScene(const UContextualAnimSceneAssetBase*
 	return true;
 }
 
-bool UContextualAnimManager::TryStartScene(const UContextualAnimSceneAssetBase* SceneAsset, AActor* PrimaryActor, const TSet<UContextualAnimSceneActorComponent*>& SceneActorComps)
+bool UContextualAnimManager::TryStartScene(const UContextualAnimSceneAsset* SceneAsset, AActor* PrimaryActor, const TSet<UContextualAnimSceneActorComponent*>& SceneActorComps)
 {
 	if (SceneAsset == nullptr)
 	{
diff --git a/Engine/Plugins/Experimental/Animation/ContextualAnimation/Source/ContextualAnimation/Private/ContextualAnimMetadata.cpp b/Engine/Plugins/Experimental/Animation/ContextualAnimation/Source/ContextualAnimation/Private/ContextualAnimMetadata.cpp
--- a/Engine/Plugins/Experimental/Animation/ContextualAnimation/Source/ContextualAnimation/Private/ContextualAnimMetadata.cpp
+++ b/Engine/Plugins/Experimental/Animation/ContextualAnimation/Source/ContextualAnimation/Private/ContextualAnimMetadata.cpp
@@ -26,9 +26,9 @@ UContextualAnimMetadata::UContextualAnimMetadata(const FObjectInitializer& Objec
 {
 }
 
-UContextualAnimSceneAssetBase* UContextualAnimMetadata::GetSceneAssetOwner() const
+UContextualAnimSceneAsset* UContextualAnimMetadata::GetSceneAssetOwner() const
 {
-	return Cast<UContextualAnimSceneAssetBase>(GetOuter());
+	return Cast<UContextualAnimSceneAsset>(GetOuter());
 }
 
 bool UContextualAnimMetadata::DoesQuerierPassConditions(const FContextualAnimQuerier& Querier, const FContextualAnimQueryContext& Context, const FTransform& EntryTransform) const
diff --git a/Engine/Plugins/Experimental/Animation/ContextualAnimation/Source/ContextualAnimation/Private/ContextualAnimSceneActorComponent.cpp b/Engine/Plugins/Experimental/Animation/ContextualAnimation/Source/ContextualAnimation/Private/ContextualAnimSceneActorComponent.cpp
--- a/Engine/Plugins/Experimental/Animation/ContextualAnimation/Source/ContextualAnimation/Private/ContextualAnimSceneActorComponent.cpp
+++ b/Engine/Plugins/Experimental/Animation/ContextualAnimation/Source/ContextualAnimation/Private/ContextualAnimSceneActorComponent.cpp
@@ -8,7 +8,6 @@
 #include "DrawDebugHelpers.h"
 #include "ContextualAnimManager.h"
 #include "ContextualAnimSceneInstance.h"
-#include "ContextualAnimCompositeSceneAsset.h"
 #include "ContextualAnimSceneAsset.h"
 #include "ContextualAnimUtilities.h"
 
@@ -201,7 +200,7 @@ FPrimitiveSceneProxy* UContextualAnimSceneActorComponent::CreateSceneProxy()
 
 		virtual void GetDynamicMeshElements(const TArray<const FSceneView*>& Views, const FSceneViewFamily& ViewFamily, uint32 VisibilityMap, FMeshElementCollector& Collector) const override
 		{
-			const UContextualAnimSceneAssetBase* Asset = SceneAssetPtr.Get();
+			const UContextualAnimSceneAsset* Asset = SceneAssetPtr.Get();
 			if (Asset == nullptr)
 			{
 				return;
@@ -346,7 +345,7 @@ FPrimitiveSceneProxy* UContextualAnimSceneActorComponent::CreateSceneProxy()
 		}
 
 	private:
-		TWeakObjectPtr<const UContextualAnimSceneAssetBase> SceneAssetPtr;
+		TWeakObjectPtr<const UContextualAnimSceneAsset> SceneAssetPtr;
 		FContextualAnimDebugParams Params;
 	};
 
diff --git a/Engine/Plugins/Experimental/Animation/ContextualAnimation/Source/ContextualAnimation/Private/ContextualAnimSceneAsset.cpp b/Engine/Plugins/Experimental/Animation/ContextualAnimation/Source/ContextualAnimation/Private/ContextualAnimSceneAsset.cpp
--- a/Engine/Plugins/Experimental/Animation/ContextualAnimation/Source/ContextualAnimation/Private/ContextualAnimSceneAsset.cpp
+++ b/Engine/Plugins/Experimental/Animation/ContextualAnimation/Source/ContextualAnimation/Private/ContextualAnimSceneAsset.cpp
@@ -9,10 +9,28 @@
 #include "ContextualAnimMetadata.h"
 #include "UObject/ObjectSaveContext.h"
 #include "ContextualAnimScenePivotProvider.h"
+#include "Containers/ArrayView.h"
+
+static FCompactPoseBoneIndex GetCompactPoseBoneIndexFromPose(const FCSPose<FCompactPose>& Pose, const FName& BoneName)
+{
+	const FBoneContainer& BoneContainer = Pose.GetPose().GetBoneContainer();
+	for (int32 Idx = Pose.GetPose().GetNumBones() - 1; Idx >= 0; Idx--)
+	{
+		if (BoneContainer.GetReferenceSkeleton().GetBoneName(BoneContainer.GetBoneIndicesArray()[Idx]) == BoneName)
+		{
+			return FCompactPoseBoneIndex(Idx);
+		}
+	}
+
+	checkf(false, TEXT("BoneName: %s Pose.Asset: %s Pose.NumBones: %d"), *BoneName.ToString(), *GetNameSafe(Pose.GetPose().GetBoneContainer().GetAsset()), Pose.GetPose().GetNumBones());
+	return FCompactPoseBoneIndex(INDEX_NONE);
+}
 
 UContextualAnimSceneAsset::UContextualAnimSceneAsset(const FObjectInitializer& ObjectInitializer)
 	: Super(ObjectInitializer)
 {
+	bDisableCollisionBetweenActors = true;
+	SampleRate = 15;
 }
 
 void UContextualAnimSceneAsset::PreSave(FObjectPreSaveContext ObjectSaveContext)
@@ -72,6 +90,321 @@ void UContextualAnimSceneAsset::PreSave(FObjectPreSaveContext ObjectSaveContext)
 	UpdateRadius();
 }
 
+void UContextualAnimSceneAsset::GenerateAlignmentTracks(const FContextualAnimTrackSettings& Settings, FContextualAnimData& AnimData)
+{
+	const FTransform MeshToComponentInverse = Settings.MeshToComponent.Inverse();
+	const float SampleInterval = 1.f / SampleRate;
+
+	// Initialize tracks for each alignment section
+	const int32 TotalTracks = AlignmentSections.Num();
+	AnimData.AlignmentData.Initialize(TotalTracks, SampleInterval);
+	for (int32 Idx = 0; Idx < AlignmentSections.Num(); Idx++)
+	{
+		AnimData.AlignmentData.Tracks.TrackNames.Add(AlignmentSections[Idx].SectionName);
+		AnimData.AlignmentData.Tracks.AnimationTracks.AddZeroed();
+	}
+
+	if (const UAnimMontage* Animation = AnimData.Animation)
+	{
+		float Time = 0.f;
+		float EndTime = Animation->GetPlayLength();
+		int32 SampleIndex = 0;
+		while (Time < EndTime)
+		{
+			Time = FMath::Clamp(SampleIndex * SampleInterval, 0.f, EndTime);
+			SampleIndex++;
+
+			const FTransform RootTransform = MeshToComponentInverse * (UContextualAnimUtilities::ExtractRootTransformFromAnimation(Animation, Time) * AnimData.MeshToScene);
+
+			for (int32 Idx = 0; Idx < AlignmentSections.Num(); Idx++)
+			{
+				FRawAnimSequenceTrack& SceneTrack = AnimData.AlignmentData.Tracks.AnimationTracks[Idx];
+
+				const FTransform ScenePivotTransform = AlignmentSections[Idx].ScenePivots[AnimData.Index];
+				const FTransform RootRelativeToScenePivot = RootTransform.GetRelativeTransform(ScenePivotTransform);
+
+				SceneTrack.PosKeys.Add(RootRelativeToScenePivot.GetLocation());
+				SceneTrack.RotKeys.Add(RootRelativeToScenePivot.GetRotation());
+			}
+		}
+	}
+	else
+	{
+		for (int32 Idx = 0; Idx < AlignmentSections.Num(); Idx++)
+		{
+			AnimData.AlignmentData.Tracks.TrackNames.Add(AlignmentSections[Idx].SectionName);
+			FRawAnimSequenceTrack& SceneTrack = AnimData.AlignmentData.Tracks.AnimationTracks.AddZeroed_GetRef();
+
+			const FTransform RootTransform = MeshToComponentInverse * AnimData.MeshToScene;
+			const FTransform ScenePivotTransform = AlignmentSections[Idx].ScenePivots[AnimData.Index];
+			const FTransform RootRelativeToScenePivot = RootTransform.GetRelativeTransform(ScenePivotTransform);
+
+			SceneTrack.PosKeys.Add(RootRelativeToScenePivot.GetLocation());
+			SceneTrack.RotKeys.Add(RootRelativeToScenePivot.GetRotation());
+		}
+	}
+}
+
+void UContextualAnimSceneAsset::GenerateIKTargetTracks(const FContextualAnimTrackSettings& Settings, FContextualAnimData& AnimData)
+{
+	AnimData.IKTargetData.Empty();
+
+	if (Settings.IKTargetDefinitions.Num() == 0)
+	{
+		return;
+	}
+
+	if (const UAnimMontage* Animation = AnimData.Animation)
+	{
+		UE_LOG(LogContextualAnim, Log, TEXT("%s Generating IK Target Tracks. Animation: %s"), *GetNameSafe(this), *GetNameSafe(Animation));
+
+		const float SampleInterval = 1.f / SampleRate;
+
+		TArray<FBoneIndexType> RequiredBoneIndexArray;
+
+		// Helper structure to group pose extraction per target so we can extract the pose for all the bones that are relative to the same target in one pass
+		struct FPoseExtractionHelper
+		{
+			const FContextualAnimData* TargetAnimDataPtr = nullptr;
+			TArray<TTuple<FName, int32, FName, int32>> BonesData; //0: MyBoneName 1: MyBoneIndex, 2: TargetBoneName, 3 TargetBoneIndex
+		};
+		TMap<FName, FPoseExtractionHelper> PoseExtractionHelperMap;
+		PoseExtractionHelperMap.Reserve(Settings.IKTargetDefinitions.Num());
+
+		int32 TotalTracks = 0;
+		for (const FContextualAnimIKTargetDefinition& IKTargetDef : Settings.IKTargetDefinitions)
+		{
+			if (IKTargetDef.Provider == EContextualAnimIKTargetProvider::Autogenerated)
+			{
+				const FName TargetRole = IKTargetDef.AutoParams.TargetRole;
+				FPoseExtractionHelper* DataPtr = PoseExtractionHelperMap.Find(TargetRole);
+				if (DataPtr == nullptr)
+				{
+					// Find AnimData for target role. 
+					const FContextualAnimData* TargetAnimDataPtr = GetAnimDataForRoleAtIndex(TargetRole, AnimData.Index);
+					if (TargetAnimDataPtr == nullptr)
+					{
+						UE_LOG(LogContextualAnim, Warning, TEXT("\t Can't find AnimTrack for TargetRole '%s'"), *TargetRole.ToString());
+						continue;
+					}
+
+					DataPtr = &PoseExtractionHelperMap.Add(TargetRole);
+					DataPtr->TargetAnimDataPtr = TargetAnimDataPtr;
+				}
+
+				//@TODO: Replace with the actual bone. We need a way to go from IKGoalName to BoneName
+				const FName BoneName = IKTargetDef.IKGoalName;
+				const int32 BoneIndex = Animation->GetSkeleton()->GetReferenceSkeleton().FindBoneIndex(BoneName);
+				if (BoneIndex == INDEX_NONE)
+				{
+					UE_LOG(LogContextualAnim, Warning, TEXT("\t Can't find BoneIndex. BoneName: %s Animation: %s Skel: %s"),
+						*BoneName.ToString(), *GetNameSafe(Animation), *GetNameSafe(Animation->GetSkeleton()));
+
+					continue;
+				}
+
+				// Find TargetBoneIndex. Note that we add TargetBoneIndexeven if it is INDEX_NONE. In this case, my bone will be relative to the origin of the target actor. 
+				// This is to support cases where the target actor doesn't have animation or TargetBoneName is None
+				FName TargetBoneName = IKTargetDef.AutoParams.BoneName;
+				const UAnimMontage* TargetAnimation = DataPtr->TargetAnimDataPtr->Animation;
+				const int32 TargetBoneIndex = TargetAnimation ? TargetAnimation->GetSkeleton()->GetReferenceSkeleton().FindBoneIndex(TargetBoneName) : INDEX_NONE;
+				if (TargetBoneIndex == INDEX_NONE)
+				{
+					UE_LOG(LogContextualAnim, Log, TEXT("\t Can't find TargetBoneIndex. BoneName: %s Animation: %s Skel: %s. Track for this bone will be relative to the origin of the target role."),
+						*TargetBoneName.ToString(), *GetNameSafe(TargetAnimation), TargetAnimation ? *GetNameSafe(TargetAnimation->GetSkeleton()) : nullptr);
+
+					TargetBoneName = NAME_None;
+				}
+
+				RequiredBoneIndexArray.AddUnique(BoneIndex);
+
+				DataPtr->BonesData.Add(MakeTuple(BoneName, BoneIndex, TargetBoneName, TargetBoneIndex));
+				TotalTracks++;
+
+				UE_LOG(LogContextualAnim, Log, TEXT("\t Bone added for extraction. BoneName: %s (%d) TargetRole: %s TargetAnimation: %s TargetBone: %s (%d)"),
+					*BoneName.ToString(), BoneIndex, *TargetRole.ToString(), *GetNameSafe(TargetAnimation), *TargetBoneName.ToString(), TargetBoneIndex);
+			}
+		}
+
+		if (TotalTracks > 0)
+		{
+			// Complete bones chain and create bone container to extract pose from my animation
+			Animation->GetSkeleton()->GetReferenceSkeleton().EnsureParentsExistAndSort(RequiredBoneIndexArray);
+			FBoneContainer BoneContainer = FBoneContainer(RequiredBoneIndexArray, FCurveEvaluationOption(false), *Animation->GetSkeleton());
+
+			// Initialize track container
+			AnimData.IKTargetData.Initialize(TotalTracks, SampleInterval);
+
+			// Initialize lookup map to go from track name to target role and bone (the bone this track is relative to)
+			AnimData.IKTargetTrackLookupMap.Empty(TotalTracks);
+
+			float Time = 0.f;
+			float EndTime = Animation->GetPlayLength();
+			int32 SampleIndex = 0;
+			while (Time < EndTime)
+			{
+				Time = FMath::Clamp(SampleIndex * SampleInterval, 0.f, EndTime);
+				SampleIndex++;
+
+				// Extract pose from my animation
+				FCSPose<FCompactPose> ComponentSpacePose;
+				UContextualAnimUtilities::ExtractComponentSpacePose(Animation, BoneContainer, Time, false, ComponentSpacePose);
+
+				// For each target role
+				for (auto& Data : PoseExtractionHelperMap)
+				{
+					// Extract pose from target animation if any
+					FCSPose<FCompactPose> OtherComponentSpacePose;
+					TArray<FBoneIndexType> OtherRequiredBoneIndexArray;
+					FBoneContainer OtherBoneContainer;
+					const UAnimMontage* OtherAnimation = Data.Value.TargetAnimDataPtr->Animation;
+					if (OtherAnimation)
+					{
+						// Prepare array with the indices of the bones to extract from target animation
+						OtherRequiredBoneIndexArray.Reserve(Data.Value.BonesData.Num());
+						for (int32 Idx = 0; Idx < Data.Value.BonesData.Num(); Idx++)
+						{
+							const int32 TargetBoneIndex = Data.Value.BonesData[Idx].Get<3>();
+							if (TargetBoneIndex != INDEX_NONE)
+							{
+								OtherRequiredBoneIndexArray.AddUnique(TargetBoneIndex);
+							}
+						}
+
+						if (OtherRequiredBoneIndexArray.Num() > 0)
+						{
+							// Complete bones chain and create bone container to extract pose form the target animation
+							OtherAnimation->GetSkeleton()->GetReferenceSkeleton().EnsureParentsExistAndSort(OtherRequiredBoneIndexArray);
+							OtherBoneContainer = FBoneContainer(OtherRequiredBoneIndexArray, FCurveEvaluationOption(false), *OtherAnimation->GetSkeleton());
+
+							// Extract pose from target animation
+							UContextualAnimUtilities::ExtractComponentSpacePose(OtherAnimation, OtherBoneContainer, Time, false, OtherComponentSpacePose);
+						}
+					}
+
+					for (int32 Idx = 0; Idx < Data.Value.BonesData.Num(); Idx++)
+					{
+						const FName TrackName = Data.Value.BonesData[Idx].Get<0>();
+						AnimData.IKTargetData.Tracks.TrackNames.Add(TrackName);
+						AnimData.IKTargetData.Tracks.AnimationTracks.AddZeroed();
+
+						// Add entry to the lookup table
+						auto& NewItem = AnimData.IKTargetTrackLookupMap.FindOrAdd(TrackName);
+						NewItem.RoleName = Data.Key;
+						NewItem.BoneName = Data.Value.BonesData[Idx].Get<2>();
+
+						// Get bone transform from my animation
+						const FName BoneName = Data.Value.BonesData[Idx].Get<0>();
+						const FCompactPoseBoneIndex BoneIndex = GetCompactPoseBoneIndexFromPose(ComponentSpacePose, BoneName);
+						const FTransform BoneTransform = (ComponentSpacePose.GetComponentSpaceTransform(BoneIndex) * AnimData.MeshToScene);
+
+						// Get bone transform from target animation
+						FTransform OtherBoneTransform = Data.Value.TargetAnimDataPtr->MeshToScene;
+						const FName TargetBoneName = Data.Value.BonesData[Idx].Get<2>();
+						if (TargetBoneName != NAME_None)
+						{
+							const FCompactPoseBoneIndex OtherBoneIndex = GetCompactPoseBoneIndexFromPose(OtherComponentSpacePose, TargetBoneName);
+							OtherBoneTransform = (OtherComponentSpacePose.GetComponentSpaceTransform(OtherBoneIndex) * Data.Value.TargetAnimDataPtr->MeshToScene);
+						}
+
+						// Get transform relative to target
+						const FTransform BoneRelativeToOther = BoneTransform.GetRelativeTransform(OtherBoneTransform);
+
+						// Add transform to the track
+						FRawAnimSequenceTrack& NewTrack = AnimData.IKTargetData.Tracks.AnimationTracks[Idx];
+						NewTrack.PosKeys.Add(BoneRelativeToOther.GetLocation());
+						NewTrack.RotKeys.Add(BoneRelativeToOther.GetRotation());
+
+						UE_LOG(LogContextualAnim, Verbose, TEXT("\t\t Animation: %s Time: %f BoneName: %s (T: %s) Target Animation: %s TargetBoneName: %s (T: %s)"),
+							*GetNameSafe(Animation), Time, *BoneName.ToString(), *BoneTransform.GetLocation().ToString(),
+							*GetNameSafe(OtherAnimation), *TargetBoneName.ToString(), *OtherBoneTransform.GetLocation().ToString());
+					}
+				}
+			}
+		}
+	}
+}
+
+void UContextualAnimSceneAsset::UpdateRadius()
+{
+	Radius = 0.f;
+	ForEachAnimData([this](const FName& Role, const FContextualAnimData& Data)
+		{
+			Radius = FMath::Max(Radius, (Data.GetAlignmentTransformAtEntryTime().GetLocation().Size()));
+			return EContextualAnimForEachResult::Continue;
+		});
+}
+
+bool UContextualAnimSceneAsset::QueryCompositeTrack(const FContextualAnimCompositeTrack* Track, FContextualAnimQueryResult& OutResult, const FContextualAnimQueryParams& QueryParams, const FTransform& ToWorldTransform) const
+{
+	OutResult.Reset();
+
+	const FTransform QueryTransform = QueryParams.Querier.IsValid() ? QueryParams.Querier->GetActorTransform() : QueryParams.QueryTransform;
+
+	int32 DataIndex = INDEX_NONE;
+	if (Track)
+	{
+		if (QueryParams.bComplexQuery)
+		{
+			for (int32 Idx = 0; Idx < Track->AnimDataContainer.Num(); Idx++)
+			{
+				const FContextualAnimData& Data = Track->AnimDataContainer[Idx];
+
+				if (Data.Metadata)
+				{
+					const FTransform EntryTransform = Data.GetAlignmentTransformAtEntryTime() * ToWorldTransform;
+
+					if (!Data.Metadata->DoesQuerierPassConditions(FContextualAnimQuerier(QueryTransform), FContextualAnimQueryContext(ToWorldTransform), EntryTransform))
+					{
+						continue;
+					}
+				}
+
+				// Return the first item that passes all tests
+				DataIndex = Idx;
+				break;
+			}
+		}
+		else // Simple Query
+		{
+			float BestDistanceSq = MAX_FLT;
+			for (int32 Idx = 0; Idx < Track->AnimDataContainer.Num(); Idx++)
+			{
+				const FContextualAnimData& Data = Track->AnimDataContainer[Idx];
+
+				//@TODO: Convert querier location to local space instead
+				const FTransform EntryTransform = Data.GetAlignmentTransformAtEntryTime() * ToWorldTransform;
+				const float DistSq = FVector::DistSquared2D(EntryTransform.GetLocation(), QueryTransform.GetLocation());
+				if (DistSq < BestDistanceSq)
+				{
+					BestDistanceSq = DistSq;
+					DataIndex = Idx;
+				}
+			}
+		}
+
+		if (DataIndex != INDEX_NONE)
+		{
+			const FContextualAnimData& ResultData = Track->AnimDataContainer[DataIndex];
+
+			OutResult.DataIndex = DataIndex;
+			OutResult.Animation = ResultData.Animation;
+			OutResult.EntryTransform = ResultData.GetAlignmentTransformAtEntryTime() * ToWorldTransform;
+			OutResult.SyncTransform = ResultData.GetAlignmentTransformAtSyncTime() * ToWorldTransform;
+
+			if (QueryParams.bFindAnimStartTime)
+			{
+				const FVector LocalLocation = (QueryTransform.GetRelativeTransform(ToWorldTransform)).GetLocation();
+				OutResult.AnimStartTime = ResultData.FindBestAnimStartTime(LocalLocation);
+			}
+
+			return true;
+		}
+	}
+
+	return false;
+}
+
 const FContextualAnimTrackSettings* UContextualAnimSceneAsset::GetTrackSettings(const FName& Role) const
 {
 	const FContextualAnimCompositeTrack* Track = DataContainer.Find(Role);
diff --git a/Engine/Plugins/Experimental/Animation/ContextualAnimation/Source/ContextualAnimation/Private/ContextualAnimSceneAssetBase.cpp b/Engine/Plugins/Experimental/Animation/ContextualAnimation/Source/ContextualAnimation/Private/ContextualAnimSceneAssetBase.cpp
deleted file mode 100644
diff --git a/Engine/Plugins/Experimental/Animation/ContextualAnimation/Source/ContextualAnimation/Private/ContextualAnimScenePivotProvider.cpp b/Engine/Plugins/Experimental/Animation/ContextualAnimation/Source/ContextualAnimation/Private/ContextualAnimScenePivotProvider.cpp
--- a/Engine/Plugins/Experimental/Animation/ContextualAnimation/Source/ContextualAnimation/Private/ContextualAnimScenePivotProvider.cpp
+++ b/Engine/Plugins/Experimental/Animation/ContextualAnimation/Source/ContextualAnimation/Private/ContextualAnimScenePivotProvider.cpp
@@ -12,9 +12,9 @@ UContextualAnimScenePivotProvider::UContextualAnimScenePivotProvider(const FObje
 {
 }
 
-const UContextualAnimSceneAssetBase* UContextualAnimScenePivotProvider::GetSceneAsset() const
+const UContextualAnimSceneAsset* UContextualAnimScenePivotProvider::GetSceneAsset() const
 {
-	return CastChecked<UContextualAnimSceneAssetBase>(GetOuter());
+	return CastChecked<UContextualAnimSceneAsset>(GetOuter());
 }
 
 // UContextualAnimScenePivotProvider_Default
diff --git a/Engine/Plugins/Experimental/Animation/ContextualAnimation/Source/ContextualAnimation/Public/ContextualAnimCompositeSceneAsset.h b/Engine/Plugins/Experimental/Animation/ContextualAnimation/Source/ContextualAnimation/Public/ContextualAnimCompositeSceneAsset.h
deleted file mode 100644
diff --git a/Engine/Plugins/Experimental/Animation/ContextualAnimation/Source/ContextualAnimation/Public/ContextualAnimManager.h b/Engine/Plugins/Experimental/Animation/ContextualAnimation/Source/ContextualAnimation/Public/ContextualAnimManager.h
--- a/Engine/Plugins/Experimental/Animation/ContextualAnimation/Source/ContextualAnimation/Public/ContextualAnimManager.h
+++ b/Engine/Plugins/Experimental/Animation/ContextualAnimation/Source/ContextualAnimation/Public/ContextualAnimManager.h
@@ -9,7 +9,7 @@
 #include "ContextualAnimManager.generated.h"
 
 class UContextualAnimSceneActorComponent;
-class UContextualAnimSceneAssetBase;
+class UContextualAnimSceneAsset;
 class UContextualAnimSceneInstance;
 class AActor;
 class UWorld;
@@ -40,10 +40,10 @@ public:
 	void UnregisterSceneActorComponent(UContextualAnimSceneActorComponent* SceneActorComp);
 
 	/** Attempt to start an scene instance with the supplied bindings for each role */
-	bool TryStartScene(const UContextualAnimSceneAssetBase* SceneAsset, const FContextualAnimSceneBindings& Bindings);
+	bool TryStartScene(const UContextualAnimSceneAsset* SceneAsset, const FContextualAnimSceneBindings& Bindings);
 
 	/** Attempt to start an scene instance with PrimaryActor bound to the primary role and the first component valid for each of the other roles */
-	bool TryStartScene(const UContextualAnimSceneAssetBase* SceneAsset, AActor* PrimaryActor, const TSet<UContextualAnimSceneActorComponent*>& SceneActorComps);
+	bool TryStartScene(const UContextualAnimSceneAsset* SceneAsset, AActor* PrimaryActor, const TSet<UContextualAnimSceneActorComponent*>& SceneActorComps);
 
 	UFUNCTION(BlueprintCallable, Category = "Contextual Anim|Manager")
 	UContextualAnimSceneActorComponent* FindClosestSceneActorCompToActor(const AActor* Actor) const;
@@ -63,7 +63,7 @@ public:
 	FORCEINLINE const TSet<UContextualAnimSceneActorComponent*>& GetSceneActorCompContainer() const { return SceneActorCompContainer; };
 
 	UFUNCTION(BlueprintCallable, Category = "Contextual Anim|Manager", meta = (DisplayName = "TryStartSceneWithBindings"))
-	bool BP_TryStartSceneWithBindings(const UContextualAnimSceneAssetBase* SceneAsset, const FContextualAnimSceneBindings& Bindings) { return TryStartScene(SceneAsset, Bindings); }
+	bool BP_TryStartSceneWithBindings(const UContextualAnimSceneAsset* SceneAsset, const FContextualAnimSceneBindings& Bindings) { return TryStartScene(SceneAsset, Bindings); }
 
 protected:
 
diff --git a/Engine/Plugins/Experimental/Animation/ContextualAnimation/Source/ContextualAnimation/Public/ContextualAnimMetadata.h b/Engine/Plugins/Experimental/Animation/ContextualAnimation/Source/ContextualAnimation/Public/ContextualAnimMetadata.h
--- a/Engine/Plugins/Experimental/Animation/ContextualAnimation/Source/ContextualAnimation/Public/ContextualAnimMetadata.h
+++ b/Engine/Plugins/Experimental/Animation/ContextualAnimation/Source/ContextualAnimation/Public/ContextualAnimMetadata.h
@@ -36,7 +36,7 @@ public:
 
 	UContextualAnimMetadata(const FObjectInitializer& ObjectInitializer);
 
-	class UContextualAnimSceneAssetBase* GetSceneAssetOwner() const;
+	class UContextualAnimSceneAsset* GetSceneAssetOwner() const;
 
 	bool DoesQuerierPassConditions(const FContextualAnimQuerier& Querier, const FContextualAnimQueryContext& Context, const FTransform& EntryTransform) const;
 
diff --git a/Engine/Plugins/Experimental/Animation/ContextualAnimation/Source/ContextualAnimation/Public/ContextualAnimSceneActorComponent.h b/Engine/Plugins/Experimental/Animation/ContextualAnimation/Source/ContextualAnimation/Public/ContextualAnimSceneActorComponent.h
--- a/Engine/Plugins/Experimental/Animation/ContextualAnimation/Source/ContextualAnimation/Public/ContextualAnimSceneActorComponent.h
+++ b/Engine/Plugins/Experimental/Animation/ContextualAnimation/Source/ContextualAnimation/Public/ContextualAnimSceneActorComponent.h
@@ -43,7 +43,7 @@ public:
 	FContextualAnimSceneActorCompDelegate OnLeftSceneDelegate;
 
 	UPROPERTY(EditAnywhere, BlueprintReadWrite, Category = "Config")
-	class UContextualAnimSceneAssetBase* SceneAsset;
+	class UContextualAnimSceneAsset* SceneAsset;
 
 	UPROPERTY(EditAnywhere, BlueprintReadWrite, Category = "Debug")
 	bool bEnableDebug;
diff --git a/Engine/Plugins/Experimental/Animation/ContextualAnimation/Source/ContextualAnimation/Public/ContextualAnimSceneAsset.h b/Engine/Plugins/Experimental/Animation/ContextualAnimation/Source/ContextualAnimation/Public/ContextualAnimSceneAsset.h
--- a/Engine/Plugins/Experimental/Animation/ContextualAnimation/Source/ContextualAnimation/Public/ContextualAnimSceneAsset.h
+++ b/Engine/Plugins/Experimental/Animation/ContextualAnimation/Source/ContextualAnimation/Public/ContextualAnimSceneAsset.h
@@ -2,31 +2,105 @@
 
 #pragma once
 
-#include "ContextualAnimSceneAssetBase.h"
+#include "Engine/DataAsset.h"
+#include "Templates/SubclassOf.h"
+#include "ContextualAnimTypes.h"
 #include "ContextualAnimSceneAsset.generated.h"
 
-UCLASS()
-class CONTEXTUALANIMATION_API UContextualAnimSceneAsset : public UContextualAnimSceneAssetBase
+class UContextualAnimScenePivotProvider;
+class UContextualAnimSceneInstance;
+
+USTRUCT(BlueprintType)
+struct FContextualAnimAlignmentSectionData
+{
+	GENERATED_BODY()
+
+	UPROPERTY(EditAnywhere, BlueprintReadWrite, Category = "Alignment")
+	FName SectionName = NAME_None;
+
+	UPROPERTY(EditAnywhere, Instanced, Category = "Alignment")
+	UContextualAnimScenePivotProvider* ScenePivotProvider = nullptr;
+
+	UPROPERTY(VisibleAnywhere, BlueprintReadWrite, Category = "Alignment", AdvancedDisplay)
+	TArray<FTransform> ScenePivots;
+};
+
+enum class EContextualAnimForEachResult : uint8
+{
+	Break,
+	Continue,
+};
+
+UCLASS(Blueprintable)
+class CONTEXTUALANIMATION_API UContextualAnimSceneAsset : public UDataAsset
 {
 	GENERATED_BODY()
 
 public:
 
+	typedef TFunction<EContextualAnimForEachResult(const FName& Role, const FContextualAnimData& AnimData)> FForEachAnimDataFunction;
+
 	UPROPERTY(EditAnywhere, BlueprintReadWrite, Category = "Animations")
 	TMap<FName, FContextualAnimCompositeTrack> DataContainer;
 
+	UPROPERTY(EditAnywhere, Category = "Settings")
+	TSubclassOf<UContextualAnimSceneInstance> SceneInstanceClass;
+
+	UPROPERTY(EditAnywhere, Category = "Settings")
+	bool bDisableCollisionBetweenActors;
+
+	UPROPERTY(EditAnywhere, Category = "Alignment")
+	TArray<FContextualAnimAlignmentSectionData> AlignmentSections;
+
+	UPROPERTY(EditAnywhere, Category = "Alignment", meta = (ClampMin = "1", ClampMax = "60"), AdvancedDisplay)
+	int32 SampleRate;
+
+	UPROPERTY(EditAnywhere, Category = "Transitions")
+	TArray<FContextualAnimTransitionContainer> Transitions;
+
+	//@TODO: Rename this is used for selection and move to private
+	UPROPERTY(EditAnywhere, Category = "Roles")
+	FName PrimaryRole;
+
 	UContextualAnimSceneAsset(const FObjectInitializer& ObjectInitializer);
 
 	virtual void PreSave(FObjectPreSaveContext ObjectSaveContext) override;
 
-	virtual const FContextualAnimTrackSettings* GetTrackSettings(const FName& Role) const override;
+	const FContextualAnimTrackSettings* GetTrackSettings(const FName& Role) const;
+	
+	const FContextualAnimData* GetAnimDataForRoleAtIndex(const FName& Role, int32 Index) const;
 	
-	virtual const FContextualAnimData* GetAnimDataForRoleAtIndex(const FName& Role, int32 Index) const override;
+	void ForEachAnimData(FForEachAnimDataFunction Function) const;
+
+	TArray<FName> GetRoles() const;
+
+	FORCEINLINE FName GetAlignmentSectionNameAtIndex(int32 Index) const
+	{
+		return AlignmentSections.IsValidIndex(Index) ? AlignmentSections[Index].SectionName : NAME_None;
+	}
+
+	FORCEINLINE const FName& GetLeaderRole() const { return LeaderRole != NAME_None ? LeaderRole : PrimaryRole; }
+
+	FORCEINLINE float GetRadius() const { return Radius; }
+
+	bool Query(const FName& Role, FContextualAnimQueryResult& OutResult, const FContextualAnimQueryParams& QueryParams, const FTransform& ToWorldTransform) const;
+
+	bool QueryCompositeTrack(const FContextualAnimCompositeTrack* Track, FContextualAnimQueryResult& OutResult, const FContextualAnimQueryParams& QueryParams, const FTransform& ToWorldTransform) const;
+
+protected:
+
+	void GenerateAlignmentTracks(const FContextualAnimTrackSettings& Settings, FContextualAnimData& AnimData);
+
+	void GenerateIKTargetTracks(const FContextualAnimTrackSettings& Settings, FContextualAnimData& AnimData);
+
+	void UpdateRadius();
 
-	virtual void ForEachAnimData(FForEachAnimDataFunction Function) const override;
+private:
 
-	virtual bool Query(const FName& Role, FContextualAnimQueryResult& OutResult, const FContextualAnimQueryParams& QueryParams, const FTransform& ToWorldTransform) const override;
+	UPROPERTY(EditAnywhere, Category = "Roles")
+	FName LeaderRole;
 
-	//@TODO: Temp until have a roles asset
-	virtual TArray<FName> GetRoles() const override;
+	/** Radius that enclose all the entry points */
+	UPROPERTY(VisibleAnywhere, Category = "Settings")
+	float Radius = 0.f;
 };
diff --git a/Engine/Plugins/Experimental/Animation/ContextualAnimation/Source/ContextualAnimation/Public/ContextualAnimSceneAssetBase.h b/Engine/Plugins/Experimental/Animation/ContextualAnimation/Source/ContextualAnimation/Public/ContextualAnimSceneAssetBase.h
deleted file mode 100644
diff --git a/Engine/Plugins/Experimental/Animation/ContextualAnimation/Source/ContextualAnimation/Public/ContextualAnimSceneInstance.h b/Engine/Plugins/Experimental/Animation/ContextualAnimation/Source/ContextualAnimation/Public/ContextualAnimSceneInstance.h
--- a/Engine/Plugins/Experimental/Animation/ContextualAnimation/Source/ContextualAnimation/Public/ContextualAnimSceneInstance.h
+++ b/Engine/Plugins/Experimental/Animation/ContextualAnimation/Source/ContextualAnimation/Public/ContextualAnimSceneInstance.h
@@ -14,7 +14,7 @@ struct FContextualAnimData;
 struct FContextualAnimTrackSettings;
 class UContextualAnimSceneInstance;
 class UContextualAnimSceneActorComponent;
-class UContextualAnimSceneAssetBase;
+class UContextualAnimSceneAsset;
 class UAnimInstance;
 class USkeletalMeshComponent;
 class UAnimMontage;
@@ -103,7 +103,7 @@ public:
 
 	/** Scene asset this instance was created from */
 	UPROPERTY()
-	TObjectPtr<const UContextualAnimSceneAssetBase> SceneAsset;
+	TObjectPtr<const UContextualAnimSceneAsset> SceneAsset;
 
 	/** Map of roles to scene actor */
 	UPROPERTY()
diff --git a/Engine/Plugins/Experimental/Animation/ContextualAnimation/Source/ContextualAnimation/Public/ContextualAnimScenePivotProvider.h b/Engine/Plugins/Experimental/Animation/ContextualAnimation/Source/ContextualAnimation/Public/ContextualAnimScenePivotProvider.h
--- a/Engine/Plugins/Experimental/Animation/ContextualAnimation/Source/ContextualAnimation/Public/ContextualAnimScenePivotProvider.h
+++ b/Engine/Plugins/Experimental/Animation/ContextualAnimation/Source/ContextualAnimation/Public/ContextualAnimScenePivotProvider.h
@@ -7,7 +7,7 @@
 #include "ContextualAnimSceneInstance.h"
 #include "ContextualAnimScenePivotProvider.generated.h"
 
-class UContextualAnimSceneAssetBase;
+class UContextualAnimSceneAsset;
 
 UCLASS(Abstract, BlueprintType, EditInlineNew)
 class CONTEXTUALANIMATION_API UContextualAnimScenePivotProvider : public UObject
@@ -21,7 +21,7 @@ public:
 	virtual FTransform CalculateScenePivot_Source(int32 AnimDataIndex) const { return FTransform::Identity; }
 	virtual FTransform CalculateScenePivot_Runtime(const TMap<FName, FContextualAnimSceneActorData>& SceneActorMap) const { return FTransform::Identity; }
 
-	const UContextualAnimSceneAssetBase* GetSceneAsset() const;
+	const UContextualAnimSceneAsset* GetSceneAsset() const;
 };
 
 UCLASS()
diff --git a/Engine/Plugins/Experimental/Animation/ContextualAnimation/Source/ContextualAnimationEditor/Private/ContextualAnimEdMode.cpp b/Engine/Plugins/Experimental/Animation/ContextualAnimation/Source/ContextualAnimationEditor/Private/ContextualAnimEdMode.cpp
--- a/Engine/Plugins/Experimental/Animation/ContextualAnimation/Source/ContextualAnimationEditor/Private/ContextualAnimEdMode.cpp
+++ b/Engine/Plugins/Experimental/Animation/ContextualAnimation/Source/ContextualAnimationEditor/Private/ContextualAnimEdMode.cpp
@@ -10,7 +10,6 @@
 #include "ContextualAnimEdModeSettings.h"
 #include "DrawDebugHelpers.h"
 #include "ContextualAnimSceneAsset.h"
-#include "ContextualAnimCompositeSceneAsset.h"
 #include "ContextualAnimUtilities.h"
 #include "ContextualAnimManager.h"
 #include "ContextualAnimPreviewManager.h"
