Gracefully handle null UContextualAnimManager in UContextualAnimSceneActorComponent::OnRegister(). This is likely an error condition, but we can still avoid crashing the application. The component will not be registered if the world does not have a contextual animation manager.

diff --git a/Engine/Plugins/Experimental/Animation/ContextualAnimation/Source/ContextualAnimation/Private/ContextualAnimSceneActorComponent.cpp b/Engine/Plugins/Experimental/Animation/ContextualAnimation/Source/ContextualAnimation/Private/ContextualAnimSceneActorComponent.cpp
--- a/Engine/Plugins/Experimental/Animation/ContextualAnimation/Source/ContextualAnimation/Private/ContextualAnimSceneActorComponent.cpp
+++ b/Engine/Plugins/Experimental/Animation/ContextualAnimation/Source/ContextualAnimation/Private/ContextualAnimSceneActorComponent.cpp
@@ -27,13 +27,23 @@ FBoxSphereBounds UContextualAnimSceneActorComponent::CalcBounds(const FTransform
 void UContextualAnimSceneActorComponent::OnRegister()
 {
 	Super::OnRegister();
-	UContextualAnimManager::Get(GetWorld())->RegisterSceneActorComponent(this);
+	UContextualAnimManager* ContextAnimManager = UContextualAnimManager::Get(GetWorld());
+	if (ensure(!bRegistered) && ensureMsgf(ContextAnimManager, TEXT("Failed to find contextual animation manager for world '%s'"), *GetWorld()->GetName()))
+	{
+		ContextAnimManager->RegisterSceneActorComponent(this);
+		bRegistered = true;
+	}
 }
 
 void UContextualAnimSceneActorComponent::OnUnregister()
 {
 	Super::OnUnregister();
-	UContextualAnimManager::Get(GetWorld())->UnregisterSceneActorComponent(this);
+	UContextualAnimManager* ContextAnimManager = UContextualAnimManager::Get(GetWorld());
+	if (bRegistered && ensureMsgf(ContextAnimManager, TEXT("Failed to find contextual animation manager for world '%s'"), *GetWorld()->GetName()))
+	{
+		ContextAnimManager->UnregisterSceneActorComponent(this);
+		bRegistered = false;
+	}
 }
 
 void UContextualAnimSceneActorComponent::OnJoinedScene(const FContextualAnimSceneActorData* SceneActorData)
diff --git a/Engine/Plugins/Experimental/Animation/ContextualAnimation/Source/ContextualAnimation/Public/ContextualAnimSceneActorComponent.h b/Engine/Plugins/Experimental/Animation/ContextualAnimation/Source/ContextualAnimation/Public/ContextualAnimSceneActorComponent.h
--- a/Engine/Plugins/Experimental/Animation/ContextualAnimation/Source/ContextualAnimation/Public/ContextualAnimSceneActorComponent.h
+++ b/Engine/Plugins/Experimental/Animation/ContextualAnimation/Source/ContextualAnimation/Public/ContextualAnimSceneActorComponent.h
@@ -81,4 +81,8 @@ protected:
 	 */
 	UFUNCTION()
 	void OnTickPose(class USkinnedMeshComponent* SkinnedMeshComponent, float DeltaTime, bool bNeedsValidRootMotion);
+
+private:
+
+	bool bRegistered = false;
 };
