Contextual Anim - Removed ensure from OnRegister to prevent cooking process from crashing #jira none #rb trivial

diff --git a/Engine/Plugins/Experimental/Animation/ContextualAnimation/Source/ContextualAnimation/Private/ContextualAnimSceneActorComponent.cpp b/Engine/Plugins/Experimental/Animation/ContextualAnimation/Source/ContextualAnimation/Private/ContextualAnimSceneActorComponent.cpp
--- a/Engine/Plugins/Experimental/Animation/ContextualAnimation/Source/ContextualAnimation/Private/ContextualAnimSceneActorComponent.cpp
+++ b/Engine/Plugins/Experimental/Animation/ContextualAnimation/Source/ContextualAnimation/Private/ContextualAnimSceneActorComponent.cpp
@@ -27,8 +27,9 @@ FBoxSphereBounds UContextualAnimSceneActorComponent::CalcBounds(const FTransform
 void UContextualAnimSceneActorComponent::OnRegister()
 {
 	Super::OnRegister();
+
 	UContextualAnimManager* ContextAnimManager = UContextualAnimManager::Get(GetWorld());
-	if (ensure(!bRegistered) && ensureMsgf(ContextAnimManager, TEXT("Failed to find contextual animation manager for world '%s'"), *GetWorld()->GetName()))
+	if (ensure(!bRegistered) && ContextAnimManager)
 	{
 		ContextAnimManager->RegisterSceneActorComponent(this);
 		bRegistered = true;
@@ -38,8 +39,9 @@ void UContextualAnimSceneActorComponent::OnRegister()
 void UContextualAnimSceneActorComponent::OnUnregister()
 {
 	Super::OnUnregister();
+
 	UContextualAnimManager* ContextAnimManager = UContextualAnimManager::Get(GetWorld());
-	if (bRegistered && ensureMsgf(ContextAnimManager, TEXT("Failed to find contextual animation manager for world '%s'"), *GetWorld()->GetName()))
+	if (bRegistered && ContextAnimManager)
 	{
 		ContextAnimManager->UnregisterSceneActorComponent(this);
 		bRegistered = false;
