Added missing include to fix NonUnity build #jira none #rb trivial #fyi [REDACTED_HANDLE]

diff --git a/Engine/Plugins/Experimental/Animation/ContextualAnimation/Source/ContextualAnimation/Private/ContextualAnimUtilities.cpp b/Engine/Plugins/Experimental/Animation/ContextualAnimation/Source/ContextualAnimation/Private/ContextualAnimUtilities.cpp
--- a/Engine/Plugins/Experimental/Animation/ContextualAnimation/Source/ContextualAnimation/Private/ContextualAnimUtilities.cpp
+++ b/Engine/Plugins/Experimental/Animation/ContextualAnimation/Source/ContextualAnimation/Private/ContextualAnimUtilities.cpp
@@ -14,6 +14,7 @@
 #include "Misc/MemStack.h"
 #include "GameFramework/Character.h"
 #include "ContextualAnimActorInterface.h"
+#include "Components/SkeletalMeshComponent.h"
 
 void UContextualAnimUtilities::ExtractLocalSpacePose(const UAnimSequenceBase* Animation, const FBoneContainer& BoneContainer, float Time, bool bExtractRootMotion, FCompactPose& OutPose)
 {
