Contextual Anim - Create manager only for game worlds. #jira none #rb trivial #fyi [REDACTED_HANDLE]

diff --git a/Engine/Plugins/Experimental/Animation/ContextualAnimation/Source/ContextualAnimation/Private/ContextualAnimation.cpp b/Engine/Plugins/Experimental/Animation/ContextualAnimation/Source/ContextualAnimation/Private/ContextualAnimation.cpp
--- a/Engine/Plugins/Experimental/Animation/ContextualAnimation/Source/ContextualAnimation/Private/ContextualAnimation.cpp
+++ b/Engine/Plugins/Experimental/Animation/ContextualAnimation/Source/ContextualAnimation/Private/ContextualAnimation.cpp
@@ -31,7 +31,10 @@ void FContextualAnimationModule::AddReferencedObjects(FReferenceCollector& Colle
 
 void FContextualAnimationModule::OnWorldInit(UWorld* World, const UWorld::InitializationValues IVS)
 {
-	WorldToManagerMap.Add(World, NewObject<UContextualAnimManager>(World));
+	if(World && World->IsGameWorld())
+	{
+		WorldToManagerMap.Add(World, NewObject<UContextualAnimManager>(World));
+	}
 }
 
 void FContextualAnimationModule::OnWorldCleanup(UWorld* World, bool bSessionEnded, bool bCleanupResources)
