Contextual Anim
- Added console variable to toggle draw debug and to override alpha value for IK Targets
- Added blueprint function to get IKTarget by goal name
- SceneActorComps are now sorted before attempt to start an interaction in the Preview EdMode

diff --git a/Engine/Plugins/Experimental/Animation/ContextualAnimation/Source/ContextualAnimation/Private/ContextualAnimSceneActorComponent.cpp b/Engine/Plugins/Experimental/Animation/ContextualAnimation/Source/ContextualAnimation/Private/ContextualAnimSceneActorComponent.cpp
--- a/Engine/Plugins/Experimental/Animation/ContextualAnimation/Source/ContextualAnimation/Private/ContextualAnimSceneActorComponent.cpp
+++ b/Engine/Plugins/Experimental/Animation/ContextualAnimation/Source/ContextualAnimation/Private/ContextualAnimSceneActorComponent.cpp
@@ -11,6 +11,12 @@
 #include "ContextualAnimSceneAsset.h"
 #include "ContextualAnimUtilities.h"
 
+#if !(UE_BUILD_SHIPPING || UE_BUILD_TEST)
+TAutoConsoleVariable<int32> CVarContextualAnimIKDebug(TEXT("a.ContextualAnim.IK.Debug"), 0, TEXT("Draw Debug IK Targets"));
+TAutoConsoleVariable<float> CVarContextualAnimIKDrawDebugLifetime(TEXT("a.ContextualAnim.IK.DrawDebugLifetime"), 0, TEXT("Draw Debug Duration"));
+TAutoConsoleVariable<float> CVarContextualAnimIKForceAlpha(TEXT("a.ContextualAnim.IK.ForceAlpha"), -1.f, TEXT("Override Alpha value for all the targets. -1 = Disable"));
+#endif
+
 UContextualAnimSceneActorComponent::UContextualAnimSceneActorComponent(const FObjectInitializer& ObjectInitializer)
 	: Super(ObjectInitializer)
 {
@@ -100,12 +106,20 @@ void UContextualAnimSceneActorComponent::UpdateIKTargets()
 		{
 			if (UAnimInstance* AnimInstance = SceneActorDataPtr->GetAnimInstance())
 			{
+
+#if !(UE_BUILD_SHIPPING || UE_BUILD_TEST)
+				const bool bDrawDebugEnable = CVarContextualAnimIKDebug.GetValueOnGameThread() > 0;
+				const float DrawDebugDuration = CVarContextualAnimIKDrawDebugLifetime.GetValueOnGameThread();
+				FTransform IKTargetParentTransformForDebug = FTransform::Identity;
+#endif
+
 				FTransform IKTargetTransform = FTransform::Identity;
 
 				float Alpha = 0.f;
 				bool bCurveFound = AnimInstance->GetCurveValue(IKTargetDef.AlphaCurveName, Alpha);
 
-				//@TODO: IKTargetTransform will be off by 1 frame if we tick before target until setup tick dependencies
+				// @TODO: IKTargetTransform will be off by 1 frame if we tick before target. 
+				// Should we at least add an option to the SceneAsset to setup tick dependencies or should this be entirely up to the user?
 
 				if (IKTargetDef.Provider == EContextualAnimIKTargetProvider::Autogenerated)
 				{
@@ -117,6 +131,13 @@ void UContextualAnimSceneActorComponent::UpdateIKTargets()
 
 							const float Time = SceneActorDataPtr->GetAnimTime();
 							IKTargetTransform = SceneActorDataPtr->GetAnimData()->IKTargetData.ExtractTransformAtTime(IKTargetDef.IKGoalName, Time) * IKTargetParentTransform;
+
+#if !(UE_BUILD_SHIPPING || UE_BUILD_TEST)
+							if (bDrawDebugEnable)
+							{
+								IKTargetParentTransformForDebug = IKTargetParentTransform;
+							}
+#endif
 						}
 					}
 				}
@@ -127,17 +148,32 @@ void UContextualAnimSceneActorComponent::UpdateIKTargets()
 						if (USkeletalMeshComponent* TargetSkelMeshComp = TargetSceneActorData->GetSkeletalMeshComponent())
 						{
 							IKTargetTransform = TargetSkelMeshComp->GetSocketTransform(IKTargetDef.BoneParams.BoneName);
+
+#if !(UE_BUILD_SHIPPING || UE_BUILD_TEST)
+							if (bDrawDebugEnable)
+							{
+								IKTargetParentTransformForDebug = TargetSkelMeshComp->GetSocketTransform(TargetSkelMeshComp->GetParentBone(IKTargetDef.BoneParams.BoneName));
+							}
+#endif
 						}
 					}
 				}
 
-				// Draw Debug IK Targets
 #if !(UE_BUILD_SHIPPING || UE_BUILD_TEST)
-				if (FContextualAnimCVars::CVarDrawDebugIKTargets.GetValueOnGameThread() > 0)
+				const float ForcedAlphaValue = CVarContextualAnimIKForceAlpha.GetValueOnGameThread();
+				if (ForcedAlphaValue > 0.f)
+				{
+					Alpha = FMath::Clamp(ForcedAlphaValue, 0.f, 1.f);
+				}
+
+				if (bDrawDebugEnable)
 				{
+					const float DrawThickness = 0.5f;
 					const FColor DrawColor = FColor::MakeRedToGreenColorFromScalar(Alpha);
-					DrawDebugSphere(GetWorld(), IKTargetTransform.GetLocation(), 5.f, 12, DrawColor, false, 0.f, 0, 0.5f);
-					DrawDebugCoordinateSystem(GetWorld(), IKTargetTransform.GetLocation(), IKTargetTransform.Rotator(), 10.f, false, 0.f, 0, 0.5f);
+
+					DrawDebugLine(GetWorld(), IKTargetParentTransformForDebug.GetLocation(), IKTargetTransform.GetLocation(), DrawColor, false, DrawDebugDuration, 0, DrawThickness);
+					DrawDebugCoordinateSystem(GetWorld(), IKTargetTransform.GetLocation(), IKTargetTransform.Rotator(), 10.f, false, DrawDebugDuration, 0, DrawThickness);
+					//DrawDebugSphere(GetWorld(), IKTargetTransform.GetLocation(), 5.f, 12, DrawColor, false, 0.f, 0, DrawThickness  );
 
 					//DrawDebugString(GetWorld(), IKTargetTransform.GetLocation(), FString::Printf(TEXT("%s (%f)"), *IKTargetDef.AlphaCurveName.ToString(), Alpha));
 				}
@@ -146,12 +182,21 @@ void UContextualAnimSceneActorComponent::UpdateIKTargets()
 				// Convert IK Target to mesh space
 				IKTargetTransform = IKTargetTransform.GetRelativeTransform(SceneActorDataPtr->GetSkeletalMeshComponent()->GetComponentTransform());
 
-				IKTargets.Add(FContextualAnimIKTarget(IKTargetDef.IKGoalName, Alpha, IKTargetTransform));
+				IKTargets.Add(FContextualAnimIKTarget(IKTargetDef.IKGoalName, Alpha, IKTargetTransform));	
 			}
 		}
 	}
 }
 
+const FContextualAnimIKTarget& UContextualAnimSceneActorComponent::GetIKTargetByGoalName(FName GoalName) const
+{
+	const FContextualAnimIKTarget* IKTargetPtr = IKTargets.FindByPredicate([GoalName](const FContextualAnimIKTarget& IKTarget){
+		return IKTarget.GoalName == GoalName;
+	});
+
+	return IKTargetPtr ? *IKTargetPtr : FContextualAnimIKTarget::InvalidIKTarget;
+}
+
 FPrimitiveSceneProxy* UContextualAnimSceneActorComponent::CreateSceneProxy()
 {
 	class FSceneActorCompProxy final : public FPrimitiveSceneProxy
diff --git a/Engine/Plugins/Experimental/Animation/ContextualAnimation/Source/ContextualAnimation/Private/ContextualAnimTypes.cpp b/Engine/Plugins/Experimental/Animation/ContextualAnimation/Source/ContextualAnimation/Private/ContextualAnimTypes.cpp
--- a/Engine/Plugins/Experimental/Animation/ContextualAnimation/Source/ContextualAnimation/Private/ContextualAnimTypes.cpp
+++ b/Engine/Plugins/Experimental/Animation/ContextualAnimation/Source/ContextualAnimation/Private/ContextualAnimTypes.cpp
@@ -10,11 +10,8 @@
 
 DEFINE_LOG_CATEGORY(LogContextualAnim);
 
-#if !(UE_BUILD_SHIPPING || UE_BUILD_TEST)
-TAutoConsoleVariable<int32> FContextualAnimCVars::CVarDrawDebugIKTargets(TEXT("a.ContextualAnim.DrawDebugIKTargets"), 0, TEXT("Draw Debug IK Targets"), ECVF_Cheat);
-#endif
-
-const FContextualAnimData FContextualAnimData::EmptyAnimData = FContextualAnimData();
+const FContextualAnimData FContextualAnimData::EmptyAnimData;
+const FContextualAnimIKTarget FContextualAnimIKTarget::InvalidIKTarget;
 
 // FContextualAnimAlignmentTrackContainer
 ///////////////////////////////////////////////////////////////////////
diff --git a/Engine/Plugins/Experimental/Animation/ContextualAnimation/Source/ContextualAnimation/Public/ContextualAnimSceneActorComponent.h b/Engine/Plugins/Experimental/Animation/ContextualAnimation/Source/ContextualAnimation/Public/ContextualAnimSceneActorComponent.h
--- a/Engine/Plugins/Experimental/Animation/ContextualAnimation/Source/ContextualAnimation/Public/ContextualAnimSceneActorComponent.h
+++ b/Engine/Plugins/Experimental/Animation/ContextualAnimation/Source/ContextualAnimation/Public/ContextualAnimSceneActorComponent.h
@@ -64,6 +64,12 @@ public:
 	/** Called from the scene instance when the actor owner of this component leave an scene */
 	void OnLeftScene(const FContextualAnimSceneActorData* SceneActorData);
 
+	UFUNCTION(BlueprintPure, Category = "Contextual Anim|Scene Actor Component")
+	const TArray<FContextualAnimIKTarget>& GetIKTargets() const { return IKTargets; }
+
+	UFUNCTION(BlueprintPure, Category = "Contextual Anim|Scene Actor Component")
+	const FContextualAnimIKTarget& GetIKTargetByGoalName(FName GoalName) const;
+
 protected:
 
 	/** Ptr back to the scene actor data that represent us in the scene instance we are part of */
diff --git a/Engine/Plugins/Experimental/Animation/ContextualAnimation/Source/ContextualAnimation/Public/ContextualAnimTypes.h b/Engine/Plugins/Experimental/Animation/ContextualAnimation/Source/ContextualAnimation/Public/ContextualAnimTypes.h
--- a/Engine/Plugins/Experimental/Animation/ContextualAnimation/Source/ContextualAnimation/Public/ContextualAnimTypes.h
+++ b/Engine/Plugins/Experimental/Animation/ContextualAnimation/Source/ContextualAnimation/Public/ContextualAnimTypes.h
@@ -12,13 +12,6 @@ class AActor;
 class UAnimMontage;
 class UContextualAnimMetadata;
 
-#if !(UE_BUILD_SHIPPING || UE_BUILD_TEST)
-struct FContextualAnimCVars
-{
-	static TAutoConsoleVariable<int32> CVarDrawDebugIKTargets;
-};
-#endif
-
 /** Container for alignment tracks */
 USTRUCT()
 struct CONTEXTUALANIMATION_API FContextualAnimAlignmentTrackContainer
@@ -203,6 +196,8 @@ struct CONTEXTUALANIMATION_API FContextualAnimIKTarget
 	FContextualAnimIKTarget(){}
 	FContextualAnimIKTarget(const FName& InGoalName, float InAlpha, const FTransform& InTransform)
 		: GoalName(InGoalName), Alpha(InAlpha), Transform(InTransform){}
+
+	static const FContextualAnimIKTarget InvalidIKTarget;
 };
 
 USTRUCT(BlueprintType)
diff --git a/Engine/Plugins/Experimental/Animation/ContextualAnimation/Source/ContextualAnimationEditor/Private/ContextualAnimEdMode.cpp b/Engine/Plugins/Experimental/Animation/ContextualAnimation/Source/ContextualAnimationEditor/Private/ContextualAnimEdMode.cpp
--- a/Engine/Plugins/Experimental/Animation/ContextualAnimation/Source/ContextualAnimationEditor/Private/ContextualAnimEdMode.cpp
+++ b/Engine/Plugins/Experimental/Animation/ContextualAnimation/Source/ContextualAnimationEditor/Private/ContextualAnimEdMode.cpp
@@ -123,8 +123,16 @@ bool FContextualAnimEdMode::InputKey(FEditorViewportClient* ViewportClient, FVie
 			else
 			{
 				if (UContextualAnimSceneActorComponent* Comp = Manager->FindClosestSceneActorCompToActor(PreviewManager->TestCharacter.Get()))
-				{	
-					Manager->TryStartScene(Comp->SceneAsset, Comp->GetOwner(), Manager->GetSceneActorCompContainer());
+				{
+					TSet<UContextualAnimSceneActorComponent*> Comps = Manager->GetSceneActorCompContainer();
+
+					Comps.Sort([Comp](const UContextualAnimSceneActorComponent& A, const UContextualAnimSceneActorComponent& B){
+						const float DistSqA = FVector::DistSquared2D(Comp->GetComponentLocation(), A.GetComponentLocation());
+						const float DistSqB = FVector::DistSquared2D(Comp->GetComponentLocation(), B.GetComponentLocation());
+						return DistSqA < DistSqB;
+					});
+
+					Manager->TryStartScene(Comp->SceneAsset, Comp->GetOwner(), Comps);
 				}
 				else
 				{
