Motion Warping: Exposed RootMotionModifier's properties to blueprint and exposed debug console variables
Contextual Anim: Separated BoneName from IKGoalName in IKTargetDefinition so we can have more than one curve for the same bone + Added a few blueprint callable functions to query data from the SceneAsset

diff --git a/Engine/Plugins/Experimental/Animation/ContextualAnimation/Source/ContextualAnimation/Private/ContextualAnimSceneAsset.cpp b/Engine/Plugins/Experimental/Animation/ContextualAnimation/Source/ContextualAnimation/Private/ContextualAnimSceneAsset.cpp
--- a/Engine/Plugins/Experimental/Animation/ContextualAnimation/Source/ContextualAnimation/Private/ContextualAnimSceneAsset.cpp
+++ b/Engine/Plugins/Experimental/Animation/ContextualAnimation/Source/ContextualAnimation/Private/ContextualAnimSceneAsset.cpp
@@ -166,7 +166,7 @@ void UContextualAnimSceneAsset::GenerateIKTargetTracks(const FContextualAnimTrac
 		struct FPoseExtractionHelper
 		{
 			const FContextualAnimData* TargetAnimDataPtr = nullptr;
-			TArray<TTuple<FName, int32, FName, int32>> BonesData; //0: MyBoneName 1: MyBoneIndex, 2: TargetBoneName, 3 TargetBoneIndex
+			TArray<TTuple<FName, FName, int32, FName, int32>> BonesData; //0: GoalName, 1: MyBoneName 2: MyBoneIndex, 3: TargetBoneName, 4: TargetBoneIndex
 		};
 		TMap<FName, FPoseExtractionHelper> PoseExtractionHelperMap;
 		PoseExtractionHelperMap.Reserve(Settings.IKTargetDefinitions.Num());
@@ -192,8 +192,7 @@ void UContextualAnimSceneAsset::GenerateIKTargetTracks(const FContextualAnimTrac
 					DataPtr->TargetAnimDataPtr = TargetAnimDataPtr;
 				}
 
-				//@TODO: Replace with the actual bone. We need a way to go from IKGoalName to BoneName
-				const FName BoneName = IKTargetDef.IKGoalName;
+				const FName BoneName = IKTargetDef.BoneName;
 				const int32 BoneIndex = Animation->GetSkeleton()->GetReferenceSkeleton().FindBoneIndex(BoneName);
 				if (BoneIndex == INDEX_NONE)
 				{
@@ -218,11 +217,11 @@ void UContextualAnimSceneAsset::GenerateIKTargetTracks(const FContextualAnimTrac
 
 				RequiredBoneIndexArray.AddUnique(BoneIndex);
 
-				DataPtr->BonesData.Add(MakeTuple(BoneName, BoneIndex, TargetBoneName, TargetBoneIndex));
+				DataPtr->BonesData.Add(MakeTuple(IKTargetDef.IKGoalName, BoneName, BoneIndex, TargetBoneName, TargetBoneIndex));
 				TotalTracks++;
 
-				UE_LOG(LogContextualAnim, Log, TEXT("\t Bone added for extraction. BoneName: %s (%d) TargetRole: %s TargetAnimation: %s TargetBone: %s (%d)"),
-					*BoneName.ToString(), BoneIndex, *TargetRole.ToString(), *GetNameSafe(TargetAnimation), *TargetBoneName.ToString(), TargetBoneIndex);
+				UE_LOG(LogContextualAnim, Log, TEXT("\t Bone added for extraction. GoalName: %s BoneName: %s (%d) TargetRole: %s TargetAnimation: %s TargetBone: %s (%d)"),
+					*IKTargetDef.IKGoalName.ToString(), *BoneName.ToString(), BoneIndex, *TargetRole.ToString(), *GetNameSafe(TargetAnimation), *TargetBoneName.ToString(), TargetBoneIndex);
 			}
 		}
 
@@ -264,7 +263,7 @@ void UContextualAnimSceneAsset::GenerateIKTargetTracks(const FContextualAnimTrac
 						OtherRequiredBoneIndexArray.Reserve(Data.Value.BonesData.Num());
 						for (int32 Idx = 0; Idx < Data.Value.BonesData.Num(); Idx++)
 						{
-							const int32 TargetBoneIndex = Data.Value.BonesData[Idx].Get<3>();
+							const int32 TargetBoneIndex = Data.Value.BonesData[Idx].Get<4>();
 							if (TargetBoneIndex != INDEX_NONE)
 							{
 								OtherRequiredBoneIndexArray.AddUnique(TargetBoneIndex);
@@ -291,16 +290,16 @@ void UContextualAnimSceneAsset::GenerateIKTargetTracks(const FContextualAnimTrac
 						// Add entry to the lookup table
 						auto& NewItem = AnimData.IKTargetTrackLookupMap.FindOrAdd(TrackName);
 						NewItem.RoleName = Data.Key;
-						NewItem.BoneName = Data.Value.BonesData[Idx].Get<2>();
+						NewItem.BoneName = Data.Value.BonesData[Idx].Get<3>();
 
 						// Get bone transform from my animation
-						const FName BoneName = Data.Value.BonesData[Idx].Get<0>();
+						const FName BoneName = Data.Value.BonesData[Idx].Get<1>();
 						const FCompactPoseBoneIndex BoneIndex = GetCompactPoseBoneIndexFromPose(ComponentSpacePose, BoneName);
 						const FTransform BoneTransform = (ComponentSpacePose.GetComponentSpaceTransform(BoneIndex) * AnimData.MeshToScene);
 
 						// Get bone transform from target animation
 						FTransform OtherBoneTransform = Data.Value.TargetAnimDataPtr->MeshToScene;
-						const FName TargetBoneName = Data.Value.BonesData[Idx].Get<2>();
+						const FName TargetBoneName = Data.Value.BonesData[Idx].Get<3>();
 						if (TargetBoneName != NAME_None)
 						{
 							const FCompactPoseBoneIndex OtherBoneIndex = GetCompactPoseBoneIndexFromPose(OtherComponentSpacePose, TargetBoneName);
@@ -449,4 +448,62 @@ bool UContextualAnimSceneAsset::Query(const FName& Role, FContextualAnimQueryRes
 {
 	const FContextualAnimCompositeTrack* Track = DataContainer.Find(Role);
 	return Track ? QueryCompositeTrack(Track, OutResult, QueryParams, ToWorldTransform) : false;
+}
+
+UAnimMontage* UContextualAnimSceneAsset::GetAnimationForRoleAtIndex(FName Role, int32 Index) const
+{
+	if(const FContextualAnimCompositeTrack* Track = DataContainer.Find(Role))
+	{
+		if(Track->AnimDataContainer.IsValidIndex(Index))
+		{
+			return Track->AnimDataContainer[Index].Animation;
+		}
+	}
+
+	return nullptr;
+}
+
+FTransform UContextualAnimSceneAsset::ExtractAlignmentTransformAtTime(FName Role, int32 AnimDataIndex, float Time) const
+{
+	if (const FContextualAnimCompositeTrack* Track = DataContainer.Find(Role))
+	{
+		if (Track->AnimDataContainer.IsValidIndex(AnimDataIndex))
+		{
+			return Track->AnimDataContainer[AnimDataIndex].GetAlignmentTransformAtTime(Time);
+		}
+	}
+
+	return FTransform::Identity;
+}
+
+FTransform UContextualAnimSceneAsset::ExtractIKTargetTransformAtTime(FName Role, int32 AnimDataIndex, FName TrackName, float Time) const
+{
+	if (const FContextualAnimCompositeTrack* Track = DataContainer.Find(Role))
+	{
+		if (Track->AnimDataContainer.IsValidIndex(AnimDataIndex))
+		{
+			return Track->AnimDataContainer[AnimDataIndex].IKTargetData.ExtractTransformAtTime(TrackName, Time);
+		}
+	}
+
+	return FTransform::Identity;
+}
+
+int32 UContextualAnimSceneAsset::FindAnimIndex(FName Role, UAnimMontage* Animation) const
+{
+	int32 Result = INDEX_NONE;
+
+	if (const FContextualAnimCompositeTrack* Track = DataContainer.Find(Role))
+	{
+		for(int32 Index = 0; Index < Track->AnimDataContainer.Num(); Index++)
+		{
+			if(Track->AnimDataContainer[Index].Animation == Animation)
+			{
+				Result = Index;
+				break;
+			}
+		}
+	}
+
+	return Result;
 }
\ No newline at end of file
diff --git a/Engine/Plugins/Experimental/Animation/ContextualAnimation/Source/ContextualAnimation/Public/ContextualAnimSceneAsset.h b/Engine/Plugins/Experimental/Animation/ContextualAnimation/Source/ContextualAnimation/Public/ContextualAnimSceneAsset.h
--- a/Engine/Plugins/Experimental/Animation/ContextualAnimation/Source/ContextualAnimation/Public/ContextualAnimSceneAsset.h
+++ b/Engine/Plugins/Experimental/Animation/ContextualAnimation/Source/ContextualAnimation/Public/ContextualAnimSceneAsset.h
@@ -83,8 +83,21 @@ public:
 
 	FORCEINLINE float GetRadius() const { return Radius; }
 
+	UFUNCTION(BlueprintCallable, Category = "Contextual Anim|Scene Asset")
 	bool Query(const FName& Role, FContextualAnimQueryResult& OutResult, const FContextualAnimQueryParams& QueryParams, const FTransform& ToWorldTransform) const;
 
+	UFUNCTION(BlueprintCallable, Category = "Contextual Anim|Scene Asset")
+	UAnimMontage* GetAnimationForRoleAtIndex(FName Role, int32 Index) const;
+
+	UFUNCTION(BlueprintCallable, Category = "Contextual Anim|Scene Asset")
+	int32 FindAnimIndex(FName Role, UAnimMontage* Animation) const;
+
+	UFUNCTION(BlueprintCallable, Category = "Contextual Anim|Scene Asset")
+	FTransform ExtractAlignmentTransformAtTime(FName Role, int32 AnimDataIndex, float Time) const;
+
+	UFUNCTION(BlueprintCallable, Category = "Contextual Anim|Scene Asset")
+	FTransform ExtractIKTargetTransformAtTime(FName Role, int32 AnimDataIndex, FName TrackName, float Time) const;
+
 	bool QueryCompositeTrack(const FContextualAnimCompositeTrack* Track, FContextualAnimQueryResult& OutResult, const FContextualAnimQueryParams& QueryParams, const FTransform& ToWorldTransform) const;
 
 protected:
diff --git a/Engine/Plugins/Experimental/Animation/ContextualAnimation/Source/ContextualAnimation/Public/ContextualAnimTypes.h b/Engine/Plugins/Experimental/Animation/ContextualAnimation/Source/ContextualAnimation/Public/ContextualAnimTypes.h
--- a/Engine/Plugins/Experimental/Animation/ContextualAnimation/Source/ContextualAnimation/Public/ContextualAnimTypes.h
+++ b/Engine/Plugins/Experimental/Animation/ContextualAnimation/Source/ContextualAnimation/Public/ContextualAnimTypes.h
@@ -163,6 +163,9 @@ struct FContextualAnimIKTargetDefinition
 
 	UPROPERTY(EditAnywhere, BlueprintReadWrite, Category = "Defaults")
 	FName IKGoalName = NAME_None;
+	
+	UPROPERTY(EditAnywhere, BlueprintReadWrite, Category = "Defaults")
+	FName BoneName = NAME_None;
 
 	UPROPERTY(EditAnywhere, BlueprintReadWrite, Category = "Defaults")
 	FName AlphaCurveName = NAME_None;
diff --git a/Engine/Plugins/Experimental/Animation/MotionWarping/Source/MotionWarping/Public/MotionWarpingComponent.h b/Engine/Plugins/Experimental/Animation/MotionWarping/Source/MotionWarping/Public/MotionWarpingComponent.h
--- a/Engine/Plugins/Experimental/Animation/MotionWarping/Source/MotionWarping/Public/MotionWarpingComponent.h
+++ b/Engine/Plugins/Experimental/Animation/MotionWarping/Source/MotionWarping/Public/MotionWarpingComponent.h
@@ -17,7 +17,7 @@ class UAnimNotifyState_MotionWarping;
 DECLARE_LOG_CATEGORY_EXTERN(LogMotionWarping, Log, All);
 
 #if !(UE_BUILD_SHIPPING || UE_BUILD_TEST)
-struct FMotionWarpingCVars
+struct MOTIONWARPING_API FMotionWarpingCVars
 {
 	static TAutoConsoleVariable<int32> CVarMotionWarpingDisable;
 	static TAutoConsoleVariable<int32> CVarMotionWarpingDebug;
@@ -122,6 +122,7 @@ public:
 	int32 AddModifier(URootMotionModifier* Modifier);
 
 	/** Mark all the modifiers as Disable */
+	UFUNCTION(BlueprintCallable, Category = "Motion Warping")
 	void DisableAllRootMotionModifiers();
 
 	URootMotionModifier* AddModifierFromTemplate(URootMotionModifier* Template, const UAnimSequenceBase* Animation, float StartTime, float EndTime);
diff --git a/Engine/Plugins/Experimental/Animation/MotionWarping/Source/MotionWarping/Public/RootMotionModifier.h b/Engine/Plugins/Experimental/Animation/MotionWarping/Source/MotionWarping/Public/RootMotionModifier.h
--- a/Engine/Plugins/Experimental/Animation/MotionWarping/Source/MotionWarping/Public/RootMotionModifier.h
+++ b/Engine/Plugins/Experimental/Animation/MotionWarping/Source/MotionWarping/Public/RootMotionModifier.h
@@ -102,31 +102,31 @@ class MOTIONWARPING_API URootMotionModifier : public UObject
 public:
 
 	/** Source of the root motion we are warping */
-	UPROPERTY()
+	UPROPERTY(BlueprintReadOnly, Category = "Defaults")
 	TWeakObjectPtr<const UAnimSequenceBase> Animation = nullptr;
 
 	/** Start time of the warping window */
-	UPROPERTY()
+	UPROPERTY(BlueprintReadOnly, Category = "Defaults")
 	float StartTime = 0.f;
 
 	/** End time of the warping window */
-	UPROPERTY()
+	UPROPERTY(BlueprintReadOnly, Category = "Defaults")
 	float EndTime = 0.f;
 
 	/** Previous playback time of the animation */
-	UPROPERTY()
+	UPROPERTY(BlueprintReadOnly, Category = "Defaults")
 	float PreviousPosition = 0.f;
 
 	/** Current playback time of the animation */
-	UPROPERTY()
+	UPROPERTY(BlueprintReadOnly, Category = "Defaults")
 	float CurrentPosition = 0.f;
 
 	/** Current blend weight of the animation */
-	UPROPERTY()
+	UPROPERTY(BlueprintReadOnly, Category = "Defaults")
 	float Weight = 0.f;
 
 	/** Whether this modifier runs before the extracted root motion is converted to world space or after */
-	UPROPERTY()
+	UPROPERTY(BlueprintReadOnly, Category = "Defaults")
 	bool bInLocalSpace = false;
 
 	/** Delegate called when this modifier is activated (starts affecting the root motion) */
