Fixed multiple uninitialized script struct members found while cooking

diff --git a/Engine/Plugins/Experimental/Animation/ContextualAnimation/Source/ContextualAnimation/Public/ContextualAnimTypes.h b/Engine/Plugins/Experimental/Animation/ContextualAnimation/Source/ContextualAnimation/Public/ContextualAnimTypes.h
--- a/Engine/Plugins/Experimental/Animation/ContextualAnimation/Source/ContextualAnimation/Public/ContextualAnimTypes.h
+++ b/Engine/Plugins/Experimental/Animation/ContextualAnimation/Source/ContextualAnimation/Public/ContextualAnimTypes.h
@@ -171,7 +171,7 @@ struct FContextualAnimIKTargetDefinition
 	FName AlphaCurveName = NAME_None;
 
 	UPROPERTY(EditAnywhere, BlueprintReadWrite, Category = "Defaults")
-	EContextualAnimIKTargetProvider Provider;
+	EContextualAnimIKTargetProvider Provider = EContextualAnimIKTargetProvider::Autogenerated;
 
 	//@TODO: Hide from the UI when Type != Auto
 	UPROPERTY(EditAnywhere, BlueprintReadWrite, Category = "Defaults", meta = (EditCondition = "Provider == EContextualAnimIKTargetProvider::Autogenerated"))
