Contextual Anim Scene Asset Editor (WIP)

diff --git a/Engine/Plugins/Experimental/Animation/ContextualAnimation/Source/ContextualAnimation/Private/AnimNotifyState_IKWindow.cpp b/Engine/Plugins/Experimental/Animation/ContextualAnimation/Source/ContextualAnimation/Private/AnimNotifyState_IKWindow.cpp
new file mode 100644
--- /dev/null
+++ b/Engine/Plugins/Experimental/Animation/ContextualAnimation/Source/ContextualAnimation/Private/AnimNotifyState_IKWindow.cpp
@@ -0,0 +1,15 @@
+// Copyright Epic Games, Inc. All Rights Reserved.
+
+#include "AnimNotifyState_IKWindow.h"
+
+UAnimNotifyState_IKWindow::UAnimNotifyState_IKWindow(const FObjectInitializer& ObjectInitializer)
+	: Super(ObjectInitializer)
+{
+	BlendIn.SetBlendTime(0.25f);
+	BlendOut.SetBlendTime(0.25f);
+}
+
+FString UAnimNotifyState_IKWindow::GetNotifyName_Implementation() const
+{
+	return FString::Printf(TEXT("IK (%s)"), *GoalName.ToString());
+}
\ No newline at end of file
diff --git a/Engine/Plugins/Experimental/Animation/ContextualAnimation/Source/ContextualAnimation/Private/ContextualAnimSceneActorComponent.cpp b/Engine/Plugins/Experimental/Animation/ContextualAnimation/Source/ContextualAnimation/Private/ContextualAnimSceneActorComponent.cpp
--- a/Engine/Plugins/Experimental/Animation/ContextualAnimation/Source/ContextualAnimation/Private/ContextualAnimSceneActorComponent.cpp
+++ b/Engine/Plugins/Experimental/Animation/ContextualAnimation/Source/ContextualAnimation/Private/ContextualAnimSceneActorComponent.cpp
@@ -116,21 +116,21 @@ void UContextualAnimSceneActorComponent::UpdateIKTargets()
 				FTransform IKTargetTransform = FTransform::Identity;
 
 				float Alpha = 0.f;
-				bool bCurveFound = AnimInstance->GetCurveValue(IKTargetDef.AlphaCurveName, Alpha);
+				bool bCurveFound = AnimInstance->GetCurveValue(IKTargetDef.GoalName, Alpha);
 
 				// @TODO: IKTargetTransform will be off by 1 frame if we tick before target. 
 				// Should we at least add an option to the SceneAsset to setup tick dependencies or should this be entirely up to the user?
 
 				if (IKTargetDef.Provider == EContextualAnimIKTargetProvider::Autogenerated)
 				{
-					if (const FContextualAnimSceneActorData* TargetSceneActorData = SceneInstance->FindSceneActorDataForRole(IKTargetDef.AutoParams.TargetRole))
+					if (const FContextualAnimSceneActorData* TargetSceneActorData = SceneInstance->FindSceneActorDataForRole(IKTargetDef.TargetRoleName))
 					{
 						if (USkeletalMeshComponent* TargetSkelMeshComp = TargetSceneActorData->GetSkeletalMeshComponent())
 						{
-							const FTransform IKTargetParentTransform = TargetSkelMeshComp->GetSocketTransform(IKTargetDef.AutoParams.BoneName);
+							const FTransform IKTargetParentTransform = TargetSkelMeshComp->GetSocketTransform(IKTargetDef.TargetBoneName);
 
 							const float Time = SceneActorDataPtr->GetAnimTime();
-							IKTargetTransform = SceneActorDataPtr->GetAnimData()->IKTargetData.ExtractTransformAtTime(IKTargetDef.IKGoalName, Time) * IKTargetParentTransform;
+							IKTargetTransform = SceneActorDataPtr->GetAnimData()->IKTargetData.ExtractTransformAtTime(IKTargetDef.GoalName, Time) * IKTargetParentTransform;
 
 #if !(UE_BUILD_SHIPPING || UE_BUILD_TEST)
 							if (bDrawDebugEnable)
@@ -143,16 +143,16 @@ void UContextualAnimSceneActorComponent::UpdateIKTargets()
 				}
 				else if (IKTargetDef.Provider == EContextualAnimIKTargetProvider::Bone)
 				{
-					if (const FContextualAnimSceneActorData* TargetSceneActorData = SceneInstance->FindSceneActorDataForRole(IKTargetDef.BoneParams.TargetRole))
+					if (const FContextualAnimSceneActorData* TargetSceneActorData = SceneInstance->FindSceneActorDataForRole(IKTargetDef.TargetRoleName))
 					{
 						if (USkeletalMeshComponent* TargetSkelMeshComp = TargetSceneActorData->GetSkeletalMeshComponent())
 						{
-							IKTargetTransform = TargetSkelMeshComp->GetSocketTransform(IKTargetDef.BoneParams.BoneName);
+							IKTargetTransform = TargetSkelMeshComp->GetSocketTransform(IKTargetDef.TargetBoneName);
 
 #if !(UE_BUILD_SHIPPING || UE_BUILD_TEST)
 							if (bDrawDebugEnable)
 							{
-								IKTargetParentTransformForDebug = TargetSkelMeshComp->GetSocketTransform(TargetSkelMeshComp->GetParentBone(IKTargetDef.BoneParams.BoneName));
+								IKTargetParentTransformForDebug = TargetSkelMeshComp->GetSocketTransform(TargetSkelMeshComp->GetParentBone(IKTargetDef.TargetBoneName));
 							}
 #endif
 						}
@@ -182,7 +182,7 @@ void UContextualAnimSceneActorComponent::UpdateIKTargets()
 				// Convert IK Target to mesh space
 				IKTargetTransform = IKTargetTransform.GetRelativeTransform(SceneActorDataPtr->GetSkeletalMeshComponent()->GetComponentTransform());
 
-				IKTargets.Add(FContextualAnimIKTarget(IKTargetDef.IKGoalName, Alpha, IKTargetTransform));	
+				IKTargets.Add(FContextualAnimIKTarget(IKTargetDef.GoalName, Alpha, IKTargetTransform));	
 			}
 		}
 	}
diff --git a/Engine/Plugins/Experimental/Animation/ContextualAnimation/Source/ContextualAnimation/Private/ContextualAnimSceneAsset.cpp b/Engine/Plugins/Experimental/Animation/ContextualAnimation/Source/ContextualAnimation/Private/ContextualAnimSceneAsset.cpp
--- a/Engine/Plugins/Experimental/Animation/ContextualAnimation/Source/ContextualAnimation/Private/ContextualAnimSceneAsset.cpp
+++ b/Engine/Plugins/Experimental/Animation/ContextualAnimation/Source/ContextualAnimation/Private/ContextualAnimSceneAsset.cpp
@@ -34,6 +34,13 @@ UContextualAnimSceneAsset::UContextualAnimSceneAsset(const FObjectInitializer& O
 }
 
 void UContextualAnimSceneAsset::PreSave(FObjectPreSaveContext ObjectSaveContext)
+{
+	Super::PreSave(ObjectSaveContext);
+
+	PrecomputeData();
+}
+
+void UContextualAnimSceneAsset::PrecomputeData()
 {
 	// Necessary for FCompactPose that uses a FAnimStackAllocator (TMemStackAllocator) which allocates from FMemStack.
 	// When allocating memory from FMemStack we need to explicitly use FMemMark to ensure items are freed when the scope exits. 
@@ -41,8 +48,6 @@ void UContextualAnimSceneAsset::PreSave(FObjectPreSaveContext ObjectSaveContext)
 	// but any allocation from outside the game tick (like here when generating the alignment tracks off-line) must explicitly add a mark to avoid a leak 
 	FMemMark Mark(FMemStack::Get());
 
-	Super::PreSave(ObjectSaveContext);
-
 	// Set Index for each ContextualAnimData
 	int32 NumAnimData = 0;
 	for (auto& Entry : DataContainer)
@@ -73,7 +78,6 @@ void UContextualAnimSceneAsset::PreSave(FObjectPreSaveContext ObjectSaveContext)
 			}
 		}
 	}
-	
 
 	for (auto& Entry : DataContainer)
 	{
@@ -176,7 +180,7 @@ void UContextualAnimSceneAsset::GenerateIKTargetTracks(const FContextualAnimTrac
 		{
 			if (IKTargetDef.Provider == EContextualAnimIKTargetProvider::Autogenerated)
 			{
-				const FName TargetRole = IKTargetDef.AutoParams.TargetRole;
+				const FName TargetRole = IKTargetDef.TargetRoleName;
 				FPoseExtractionHelper* DataPtr = PoseExtractionHelperMap.Find(TargetRole);
 				if (DataPtr == nullptr)
 				{
@@ -204,7 +208,7 @@ void UContextualAnimSceneAsset::GenerateIKTargetTracks(const FContextualAnimTrac
 
 				// Find TargetBoneIndex. Note that we add TargetBoneIndexeven if it is INDEX_NONE. In this case, my bone will be relative to the origin of the target actor. 
 				// This is to support cases where the target actor doesn't have animation or TargetBoneName is None
-				FName TargetBoneName = IKTargetDef.AutoParams.BoneName;
+				FName TargetBoneName = IKTargetDef.TargetBoneName;
 				const UAnimMontage* TargetAnimation = DataPtr->TargetAnimDataPtr->Animation;
 				const int32 TargetBoneIndex = TargetAnimation ? TargetAnimation->GetSkeleton()->GetReferenceSkeleton().FindBoneIndex(TargetBoneName) : INDEX_NONE;
 				if (TargetBoneIndex == INDEX_NONE)
@@ -217,11 +221,11 @@ void UContextualAnimSceneAsset::GenerateIKTargetTracks(const FContextualAnimTrac
 
 				RequiredBoneIndexArray.AddUnique(BoneIndex);
 
-				DataPtr->BonesData.Add(MakeTuple(IKTargetDef.IKGoalName, BoneName, BoneIndex, TargetBoneName, TargetBoneIndex));
+				DataPtr->BonesData.Add(MakeTuple(IKTargetDef.GoalName, BoneName, BoneIndex, TargetBoneName, TargetBoneIndex));
 				TotalTracks++;
 
 				UE_LOG(LogContextualAnim, Log, TEXT("\t Bone added for extraction. GoalName: %s BoneName: %s (%d) TargetRole: %s TargetAnimation: %s TargetBone: %s (%d)"),
-					*IKTargetDef.IKGoalName.ToString(), *BoneName.ToString(), BoneIndex, *TargetRole.ToString(), *GetNameSafe(TargetAnimation), *TargetBoneName.ToString(), TargetBoneIndex);
+					*IKTargetDef.GoalName.ToString(), *BoneName.ToString(), BoneIndex, *TargetRole.ToString(), *GetNameSafe(TargetAnimation), *TargetBoneName.ToString(), TargetBoneIndex);
 			}
 		}
 
@@ -334,6 +338,38 @@ void UContextualAnimSceneAsset::UpdateRadius()
 		});
 }
 
+FName UContextualAnimSceneAsset::FindRoleByAnimation(const UAnimMontage* Animation) const
+{
+	for (const auto& Pair : DataContainer)
+	{
+		for (const FContextualAnimData& Data : Pair.Value.AnimDataContainer)
+		{
+			if (Data.Animation == Animation)
+			{
+				return Pair.Key;
+			}
+		}
+	}
+
+	return NAME_None;
+}
+
+TArray<FName> UContextualAnimSceneAsset::GetIKGoalsForRole(const FName& Role) const
+{
+	TArray<FName> Goals;
+
+	if(const FContextualAnimTrackSettings* Settings = GetTrackSettings(Role))
+	{
+		Goals.Reserve(Settings->IKTargetDefinitions.Num());
+		for(const FContextualAnimIKTargetDefinition& IKTargetDef : Settings->IKTargetDefinitions)
+		{
+			Goals.Add(IKTargetDef.GoalName);
+		}
+	}
+
+	return Goals;
+}
+
 bool UContextualAnimSceneAsset::QueryCompositeTrack(const FContextualAnimCompositeTrack* Track, FContextualAnimQueryResult& OutResult, const FContextualAnimQueryParams& QueryParams, const FTransform& ToWorldTransform) const
 {
 	OutResult.Reset();
diff --git a/Engine/Plugins/Experimental/Animation/ContextualAnimation/Source/ContextualAnimation/Private/ContextualAnimUtilities.cpp b/Engine/Plugins/Experimental/Animation/ContextualAnimation/Source/ContextualAnimation/Private/ContextualAnimUtilities.cpp
--- a/Engine/Plugins/Experimental/Animation/ContextualAnimation/Source/ContextualAnimation/Private/ContextualAnimUtilities.cpp
+++ b/Engine/Plugins/Experimental/Animation/ContextualAnimation/Source/ContextualAnimation/Private/ContextualAnimUtilities.cpp
@@ -7,6 +7,7 @@
 #include "Animation/AnimationPoseData.h"
 #include "Animation/AttributesRuntime.h"
 #include "Animation/AnimTypes.h"
+#include "Animation/AnimInstance.h"
 #include "DrawDebugHelpers.h"
 #include "Engine/World.h"
 #include "ContextualAnimTypes.h"
@@ -15,6 +16,8 @@
 #include "GameFramework/Character.h"
 #include "ContextualAnimActorInterface.h"
 #include "Components/SkeletalMeshComponent.h"
+#include "EngineUtils.h"
+#include "AnimNotifyState_IKWindow.h"
 
 void UContextualAnimUtilities::ExtractLocalSpacePose(const UAnimSequenceBase* Animation, const FBoneContainer& BoneContainer, float Time, bool bExtractRootMotion, FCompactPose& OutPose)
 {
@@ -147,12 +150,12 @@ void UContextualAnimUtilities::DrawDebugScene(const UWorld* World, const UContex
 	}
 }
 
-USkeletalMeshComponent* UContextualAnimUtilities::TryGetSkeletalMeshComponent(AActor* Actor)
+USkeletalMeshComponent* UContextualAnimUtilities::TryGetSkeletalMeshComponent(const AActor* Actor)
 {
 	USkeletalMeshComponent* SkelMeshComp = nullptr;
 	if (Actor)
 	{
-		if (ACharacter* Character = Cast<ACharacter>(Actor))
+		if (const ACharacter* Character = Cast<const ACharacter>(Actor))
 		{
 			SkelMeshComp = Character->GetMesh();
 		}
@@ -169,7 +172,7 @@ USkeletalMeshComponent* UContextualAnimUtilities::TryGetSkeletalMeshComponent(AA
 	return SkelMeshComp;
 }
 
-UAnimInstance* UContextualAnimUtilities::TryGetAnimInstance(AActor* Actor)
+UAnimInstance* UContextualAnimUtilities::TryGetAnimInstance(const AActor* Actor)
 {
 	if (USkeletalMeshComponent* SkelMeshComp = UContextualAnimUtilities::TryGetSkeletalMeshComponent(Actor))
 	{
@@ -179,6 +182,16 @@ UAnimInstance* UContextualAnimUtilities::TryGetAnimInstance(AActor* Actor)
 	return nullptr;
 }
 
+FAnimMontageInstance* UContextualAnimUtilities::TryGetActiveAnimMontageInstance(const AActor* Actor)
+{
+	if(UAnimInstance* AnimInstance = UContextualAnimUtilities::TryGetAnimInstance(Actor))
+	{
+		return AnimInstance->GetActiveMontageInstance();
+	}
+
+	return nullptr;
+}
+
 void UContextualAnimUtilities::BP_Montage_GetSectionStartAndEndTime(const UAnimMontage* Montage, int32 SectionIndex, float& OutStartTime, float& OutEndTime)
 {
 	if(Montage)
diff --git a/Engine/Plugins/Experimental/Animation/ContextualAnimation/Source/ContextualAnimation/Public/AnimNotifyState_IKWindow.h b/Engine/Plugins/Experimental/Animation/ContextualAnimation/Source/ContextualAnimation/Public/AnimNotifyState_IKWindow.h
new file mode 100644
--- /dev/null
+++ b/Engine/Plugins/Experimental/Animation/ContextualAnimation/Source/ContextualAnimation/Public/AnimNotifyState_IKWindow.h
@@ -0,0 +1,33 @@
+// Copyright Epic Games, Inc. All Rights Reserved.
+
+#pragma once
+
+#include "CoreMinimal.h"
+#include "Animation/AnimNotifies/AnimNotifyState.h"
+#include "AlphaBlend.h"
+#include "AnimNotifyState_IKWindow.generated.h"
+
+/** AnimNotifyState used to define areas in an animation where IK should be enable */
+UCLASS(meta = (DisplayName = "IK Window"))
+class CONTEXTUALANIMATION_API UAnimNotifyState_IKWindow : public UAnimNotifyState
+{
+	GENERATED_BODY()
+
+public:
+
+	UPROPERTY(EditAnywhere, Category = Config)
+	FName GoalName;
+
+	UPROPERTY(EditAnywhere, Category = Config)
+	FAlphaBlend BlendIn;
+
+	UPROPERTY(EditAnywhere, Category = Config)
+	FAlphaBlend BlendOut;
+
+	UPROPERTY()
+	bool bEnable = true;
+
+	UAnimNotifyState_IKWindow(const FObjectInitializer& ObjectInitializer);
+
+	virtual FString GetNotifyName_Implementation() const override;
+};
\ No newline at end of file
diff --git a/Engine/Plugins/Experimental/Animation/ContextualAnimation/Source/ContextualAnimation/Public/ContextualAnimSceneAsset.h b/Engine/Plugins/Experimental/Animation/ContextualAnimation/Source/ContextualAnimation/Public/ContextualAnimSceneAsset.h
--- a/Engine/Plugins/Experimental/Animation/ContextualAnimation/Source/ContextualAnimation/Public/ContextualAnimSceneAsset.h
+++ b/Engine/Plugins/Experimental/Animation/ContextualAnimation/Source/ContextualAnimation/Public/ContextualAnimSceneAsset.h
@@ -66,6 +66,8 @@ public:
 
 	virtual void PreSave(FObjectPreSaveContext ObjectSaveContext) override;
 
+	void PrecomputeData();
+
 	const FContextualAnimTrackSettings* GetTrackSettings(const FName& Role) const;
 	
 	const FContextualAnimData* GetAnimDataForRoleAtIndex(const FName& Role, int32 Index) const;
@@ -100,6 +102,10 @@ public:
 
 	bool QueryCompositeTrack(const FContextualAnimCompositeTrack* Track, FContextualAnimQueryResult& OutResult, const FContextualAnimQueryParams& QueryParams, const FTransform& ToWorldTransform) const;
 
+	FName FindRoleByAnimation(const UAnimMontage* Animation) const;
+
+	TArray<FName> GetIKGoalsForRole(const FName& Role) const;
+
 protected:
 
 	void GenerateAlignmentTracks(const FContextualAnimTrackSettings& Settings, FContextualAnimData& AnimData);
diff --git a/Engine/Plugins/Experimental/Animation/ContextualAnimation/Source/ContextualAnimation/Public/ContextualAnimTypes.h b/Engine/Plugins/Experimental/Animation/ContextualAnimation/Source/ContextualAnimation/Public/ContextualAnimTypes.h
--- a/Engine/Plugins/Experimental/Animation/ContextualAnimation/Source/ContextualAnimation/Public/ContextualAnimTypes.h
+++ b/Engine/Plugins/Experimental/Animation/ContextualAnimation/Source/ContextualAnimation/Public/ContextualAnimTypes.h
@@ -144,42 +144,34 @@ enum class EContextualAnimIKTargetProvider : uint8
 	//@TODO: Add collision query methods
 };
 
-USTRUCT(BlueprintType)
-struct FContextualAnimIKTargetBoneParams
-{
-	GENERATED_BODY()
-
-	UPROPERTY(EditAnywhere, BlueprintReadWrite, Category = "Defaults")
-	FName TargetRole = NAME_None;
-
-	UPROPERTY(EditAnywhere, BlueprintReadWrite, Category = "Defaults")
-	FName BoneName = NAME_None;
-};
-
 USTRUCT(BlueprintType)
 struct FContextualAnimIKTargetDefinition
 {
 	GENERATED_BODY()
 
 	UPROPERTY(EditAnywhere, BlueprintReadWrite, Category = "Defaults")
-	FName IKGoalName = NAME_None;
+	FName GoalName = NAME_None;
 	
 	UPROPERTY(EditAnywhere, BlueprintReadWrite, Category = "Defaults")
 	FName BoneName = NAME_None;
 
 	UPROPERTY(EditAnywhere, BlueprintReadWrite, Category = "Defaults")
-	FName AlphaCurveName = NAME_None;
+	EContextualAnimIKTargetProvider Provider = EContextualAnimIKTargetProvider::Autogenerated;
 
 	UPROPERTY(EditAnywhere, BlueprintReadWrite, Category = "Defaults")
-	EContextualAnimIKTargetProvider Provider = EContextualAnimIKTargetProvider::Autogenerated;
+	FName TargetRoleName = NAME_None;
 
-	//@TODO: Hide from the UI when Type != Auto
-	UPROPERTY(EditAnywhere, BlueprintReadWrite, Category = "Defaults", meta = (EditCondition = "Provider == EContextualAnimIKTargetProvider::Autogenerated"))
-	FContextualAnimIKTargetBoneParams AutoParams;
+	UPROPERTY(EditAnywhere, BlueprintReadWrite, Category = "Defaults")
+	FName TargetBoneName = NAME_None;
 
-	//@TODO: Hide from the UI when Type != Bone
-	UPROPERTY(EditAnywhere, BlueprintReadWrite, Category = "Defaults", meta = (EditCondition = "Provider == EContextualAnimIKTargetProvider::Bone"))
-	FContextualAnimIKTargetBoneParams BoneParams;
+	bool operator==(const FContextualAnimIKTargetDefinition& Other) const 
+	{
+		return  GoalName == Other.GoalName && 
+				BoneName == Other.BoneName && 
+				Provider == Other.Provider && 
+				TargetRoleName == Other.TargetRoleName &&
+				TargetBoneName == Other.TargetBoneName;
+	}
 };
 
 USTRUCT(BlueprintType)
diff --git a/Engine/Plugins/Experimental/Animation/ContextualAnimation/Source/ContextualAnimation/Public/ContextualAnimUtilities.h b/Engine/Plugins/Experimental/Animation/ContextualAnimation/Source/ContextualAnimation/Public/ContextualAnimUtilities.h
--- a/Engine/Plugins/Experimental/Animation/ContextualAnimation/Source/ContextualAnimation/Public/ContextualAnimUtilities.h
+++ b/Engine/Plugins/Experimental/Animation/ContextualAnimation/Source/ContextualAnimation/Public/ContextualAnimUtilities.h
@@ -4,12 +4,14 @@
 
 #include "Kismet/BlueprintFunctionLibrary.h"
 #include "Animation/AnimSequence.h"
+#include "ContextualAnimTypes.h"
 #include "ContextualAnimUtilities.generated.h"
 
 class UContextualAnimSceneAsset;
 class USkeletalMeshComponent;
 class UAnimInstance;
 class AActor;
+struct FAnimMontageInstance;
 
 UCLASS()
 class CONTEXTUALANIMATION_API UContextualAnimUtilities : public UBlueprintFunctionLibrary
@@ -42,9 +44,11 @@ public:
 	
 	static void DrawDebugScene(const UWorld* World, const UContextualAnimSceneAsset* SceneAsset, int32 AnimDataIndex, float Time, const FTransform& ToWorldTransform, const FColor& Color, float LifeTime, float Thickness);
 
-	static USkeletalMeshComponent* TryGetSkeletalMeshComponent(AActor* Actor);
+	static USkeletalMeshComponent* TryGetSkeletalMeshComponent(const AActor* Actor);
 
-	static UAnimInstance* TryGetAnimInstance(AActor* Actor);
+	static UAnimInstance* TryGetAnimInstance(const AActor* Actor);
+
+	static FAnimMontageInstance* TryGetActiveAnimMontageInstance(const AActor* Actor);
 
 	UFUNCTION(BlueprintCallable, Category = "Contextual Anim|Utilities", meta = (DisplayName = "GetSectionStartAndEndTime"))
 	static void BP_Montage_GetSectionStartAndEndTime(const UAnimMontage* Montage, int32 SectionIndex, float& OutStartTime, float& OutEndTime);
@@ -54,4 +58,5 @@ public:
 
 	UFUNCTION(BlueprintCallable, Category = "Contextual Anim|Utilities", meta = (DisplayName = "GetSectionLength"))
 	static float BP_Montage_GetSectionLength(const UAnimMontage* Montage, int32 SectionIndex);
+
 };
\ No newline at end of file
diff --git a/Engine/Plugins/Experimental/Animation/ContextualAnimation/Source/ContextualAnimationEditor/ContextualAnimationEditor.Build.cs b/Engine/Plugins/Experimental/Animation/ContextualAnimation/Source/ContextualAnimationEditor/ContextualAnimationEditor.Build.cs
--- a/Engine/Plugins/Experimental/Animation/ContextualAnimation/Source/ContextualAnimationEditor/ContextualAnimationEditor.Build.cs
+++ b/Engine/Plugins/Experimental/Animation/ContextualAnimation/Source/ContextualAnimationEditor/ContextualAnimationEditor.Build.cs
@@ -7,17 +7,15 @@ public class ContextualAnimationEditor : ModuleRules
 	public ContextualAnimationEditor(ReadOnlyTargetRules Target) : base(Target)
 	{
 		PCHUsage = ModuleRules.PCHUsageMode.UseExplicitOrSharedPCHs;
-		
+
 		PublicIncludePaths.AddRange(
 			new string[] {
-				// ... add public include paths required here ...
 			}
 			);
 				
 		
 		PrivateIncludePaths.AddRange(
 			new string[] {
-
 			}
 			);
 			
@@ -29,7 +27,7 @@ public class ContextualAnimationEditor : ModuleRules
 				"AIModule",
 				"NavigationSystem",
 				"ContextualAnimation",
-				"MotionWarping",
+				"MotionWarping",	
 			}
 			);
 			
@@ -47,7 +45,20 @@ public class ContextualAnimationEditor : ModuleRules
 				"LevelEditor",
 				"DetailCustomizations",
 				"EditorStyle",
-				"Projects"
+				"Projects",
+				"AssetTools",
+				"EditorWidgets",
+				"PropertyEditor",
+				"KismetWidgets",
+				"AdvancedPreviewScene",
+				"GraphEditor",
+				"SequencerWidgets",
+				"MovieScene",
+				"Sequencer",
+				"MovieSceneTracks",
+				"MovieSceneTools",
+				"ClassViewer",
+				"Persona"
 			}
 			);
 		
diff --git a/Engine/Plugins/Experimental/Animation/ContextualAnimation/Source/ContextualAnimationEditor/Private/ContextualAnimAssetEditorCommands.cpp b/Engine/Plugins/Experimental/Animation/ContextualAnimation/Source/ContextualAnimationEditor/Private/ContextualAnimAssetEditorCommands.cpp
new file mode 100644
--- /dev/null
+++ b/Engine/Plugins/Experimental/Animation/ContextualAnimation/Source/ContextualAnimationEditor/Private/ContextualAnimAssetEditorCommands.cpp
@@ -0,0 +1,12 @@
+// Copyright Epic Games, Inc. All Rights Reserved.
+
+#include "ContextualAnimAssetEditorCommands.h"
+
+#define LOCTEXT_NAMESPACE "ContextualAnimAssetEditorCommands"
+
+void FContextualAnimAssetEditorCommands::RegisterCommands()
+{
+	UI_COMMAND(ResetPreviewScene, "Reset Scene", "Reset Scene.", EUserInterfaceActionType::Button, FInputChord());
+}
+
+#undef LOCTEXT_NAMESPACE
diff --git a/Engine/Plugins/Experimental/Animation/ContextualAnimation/Source/ContextualAnimationEditor/Private/ContextualAnimAssetEditorToolkit.cpp b/Engine/Plugins/Experimental/Animation/ContextualAnimation/Source/ContextualAnimationEditor/Private/ContextualAnimAssetEditorToolkit.cpp
new file mode 100644
--- /dev/null
+++ b/Engine/Plugins/Experimental/Animation/ContextualAnimation/Source/ContextualAnimationEditor/Private/ContextualAnimAssetEditorToolkit.cpp
@@ -0,0 +1,338 @@
+// Copyright Epic Games, Inc. All Rights Reserved.
+
+#include "ContextualAnimAssetEditorToolkit.h"
+#include "SContextualAnimViewport.h"
+#include "SContextualAnimAssetBrowser.h"
+#include "ContextualAnimPreviewScene.h"
+#include "ContextualAnimAssetEditorCommands.h"
+#include "Widgets/Docking/SDockTab.h"
+#include "GameFramework/WorldSettings.h"
+#include "ContextualAnimViewModel.h"
+#include "ContextualAnimSceneAsset.h"
+#include "ContextualAnimMovieSceneSequence.h"
+#include "ISequencerModule.h"
+#include "ISequencer.h"
+#include "MovieScene.h"
+#include "Framework/MultiBox/MultiBoxBuilder.h"
+#include "AdvancedPreviewSceneModule.h"
+#include "PropertyEditorModule.h"
+#include "Modules/ModuleManager.h"
+
+#define LOCTEXT_NAMESPACE "ContextualAnimAssetEditorToolkit"
+
+const FName ContextualAnimEditorAppName = FName(TEXT("ContextualAnimEditorApp"));
+
+// Tab identifiers
+struct FContextualAnimEditorTabs
+{
+	static const FName AssetDetailsID;
+	static const FName ViewportID;
+	static const FName TimelineID;
+	static const FName AssetBrowserID;
+	static const FName PreviewSettingsID;
+};
+
+const FName FContextualAnimEditorTabs::AssetDetailsID(TEXT("ContextualAnimEditorAssetDetailsTabID"));
+const FName FContextualAnimEditorTabs::ViewportID(TEXT("ContextualAnimEditorViewportTabID"));
+const FName FContextualAnimEditorTabs::TimelineID(TEXT("ContextualAnimEditorTimelineTabID"));
+const FName FContextualAnimEditorTabs::AssetBrowserID(TEXT("ContextualAnimEditorAssetBrowserTabID"));
+const FName FContextualAnimEditorTabs::PreviewSettingsID(TEXT("ContextualAnimEditorPreviewSettingsTabID"));
+
+FContextualAnimAssetEditorToolkit::FContextualAnimAssetEditorToolkit()
+{
+}
+
+FContextualAnimAssetEditorToolkit::~FContextualAnimAssetEditorToolkit()
+{
+}
+
+UContextualAnimSceneAsset* FContextualAnimAssetEditorToolkit::GetSceneAsset() const
+{
+	return ViewModel.IsValid() ? ViewModel->GetSceneAsset() : nullptr;
+}
+
+UContextualAnimPreviewManager* FContextualAnimAssetEditorToolkit::GetPreviewManager() const
+{
+	return ViewModel.IsValid() ? ViewModel->GetPreviewManager() : nullptr;
+}
+
+void FContextualAnimAssetEditorToolkit::ResetPreviewScene()
+{
+	ViewModel->RefreshSequencerTracks();
+}
+
+void FContextualAnimAssetEditorToolkit::InitAssetEditor(const EToolkitMode::Type Mode, const TSharedPtr< IToolkitHost >& InitToolkitHost, UContextualAnimSceneAsset* SceneAsset)
+{
+	// Bind Commands
+	BindCommands();
+
+	// Create Preview Scene
+	if (!PreviewScene.IsValid())
+	{
+		PreviewScene = MakeShareable(new FContextualAnimPreviewScene(FPreviewScene::ConstructionValues().AllowAudioPlayback(true).ShouldSimulatePhysics(true).ForceUseMovementComponentInNonGameWorld(true),
+			StaticCastSharedRef<FContextualAnimAssetEditorToolkit>(AsShared())));
+
+		//Temporary fix for missing attached assets - MDW (Copied from FPersonaToolkit::CreatePreviewScene)
+		PreviewScene->GetWorld()->GetWorldSettings()->SetIsTemporarilyHiddenInEditor(false);
+	}
+
+	// Create viewport widget
+	FContextualAnimViewportRequiredArgs ViewportArgs(StaticCastSharedRef<FContextualAnimAssetEditorToolkit>(AsShared()), PreviewScene.ToSharedRef());
+	ViewportWidget = SNew(SContextualAnimViewport, ViewportArgs);
+
+	// Create asset browser widget
+	AssetBrowserWidget = SNew(SContextualAnimAssetBrowser);
+
+	ViewModel = MakeShared<FContextualAnimViewModel>();
+	ViewModel->Initialize(SceneAsset, PreviewScene.ToSharedRef());
+
+	// Create Asset Details widget
+	FPropertyEditorModule& PropertyModule = FModuleManager::LoadModuleChecked<FPropertyEditorModule>("PropertyEditor");
+
+	FDetailsViewArgs Args;
+	Args.bHideSelectionTip = true;
+	Args.NotifyHook = this;
+
+	EditingAssetWidget = PropertyModule.CreateDetailView(Args);
+	EditingAssetWidget->SetObject(SceneAsset);
+	EditingAssetWidget->OnFinishedChangingProperties().AddSP(this, &FContextualAnimAssetEditorToolkit::OnFinishedChangingProperties);
+
+	// Define Editor Layout
+	const TSharedRef<FTabManager::FLayout> StandaloneDefaultLayout = FTabManager::NewLayout("Standalone_ContextualAnimAnimEditor_Layout_v0.07")
+		->AddArea
+		(
+			FTabManager::NewPrimaryArea()->SetOrientation(Orient_Vertical)
+			->Split
+			(
+				FTabManager::NewSplitter()->SetOrientation(Orient_Vertical)->SetSizeCoefficient(0.9f)
+				->Split
+				(
+					FTabManager::NewSplitter()->SetOrientation(Orient_Horizontal)->SetSizeCoefficient(0.9f)
+					->Split
+					(
+						FTabManager::NewSplitter()->SetOrientation(Orient_Vertical)
+						->Split
+						(
+							FTabManager::NewStack()
+							->SetSizeCoefficient(0.65f)
+							->AddTab(FContextualAnimEditorTabs::ViewportID, ETabState::OpenedTab)->SetHideTabWell(true)
+						)
+						->Split
+						(
+							FTabManager::NewStack()
+							->SetSizeCoefficient(0.3f)
+							->AddTab(FContextualAnimEditorTabs::TimelineID, ETabState::OpenedTab)
+						)
+					)
+					->Split
+					(
+						FTabManager::NewSplitter()->SetOrientation(Orient_Vertical)
+						->Split
+						(
+							FTabManager::NewStack()
+							->SetSizeCoefficient(0.3f)
+							->AddTab(FContextualAnimEditorTabs::AssetDetailsID, ETabState::OpenedTab)
+							->AddTab(FContextualAnimEditorTabs::PreviewSettingsID, ETabState::OpenedTab)
+						)
+						->Split
+						(
+							FTabManager::NewStack()
+							->SetSizeCoefficient(0.3f)
+							->AddTab(FContextualAnimEditorTabs::AssetBrowserID, ETabState::OpenedTab)
+						)
+					)
+				)
+			)
+		);
+
+	const bool bCreateDefaultStandaloneMenuParam = true;
+	const bool bCreateDefaultToolbarParam = true;
+	const bool bIsToolbarFocusableParam = false;
+	FAssetEditorToolkit::InitAssetEditor(Mode, InitToolkitHost, ContextualAnimEditorAppName, StandaloneDefaultLayout, bCreateDefaultStandaloneMenuParam, bCreateDefaultToolbarParam, SceneAsset, bIsToolbarFocusableParam);
+
+	ExtendToolbar();
+
+	RegenerateMenusAndToolbars();
+}
+
+void FContextualAnimAssetEditorToolkit::BindCommands()
+{
+	const FContextualAnimAssetEditorCommands& Commands = FContextualAnimAssetEditorCommands::Get();
+
+	ToolkitCommands->MapAction(
+		Commands.ResetPreviewScene,
+		FExecuteAction::CreateSP(this, &FContextualAnimAssetEditorToolkit::ResetPreviewScene),
+		EUIActionRepeatMode::RepeatDisabled);
+}
+
+void FContextualAnimAssetEditorToolkit::ExtendToolbar()
+{
+	TSharedPtr<FExtender> ToolbarExtender = MakeShareable(new FExtender);
+
+	AddToolbarExtender(ToolbarExtender);
+
+	ToolbarExtender->AddToolBarExtension(
+		"Asset",
+		EExtensionHook::After,
+		GetToolkitCommands(),
+		FToolBarExtensionDelegate::CreateSP(this, &FContextualAnimAssetEditorToolkit::FillToolbar)
+	);
+}
+
+void FContextualAnimAssetEditorToolkit::FillToolbar(FToolBarBuilder& ToolbarBuilder)
+{
+	ToolbarBuilder.BeginSection("Reset");
+	{
+		ToolbarBuilder.AddToolBarButton(
+			FContextualAnimAssetEditorCommands::Get().ResetPreviewScene,
+			NAME_None,
+			TAttribute<FText>(),
+			TAttribute<FText>(),
+			FSlateIcon(FAppStyle::Get().GetStyleSetName(), "Icons.Refresh"));
+	}
+	ToolbarBuilder.EndSection();
+}
+
+void FContextualAnimAssetEditorToolkit::RegisterTabSpawners(const TSharedRef<FTabManager>& InTabManager)
+{
+	WorkspaceMenuCategory = InTabManager->AddLocalWorkspaceMenuCategory(LOCTEXT("WorkspaceMenu_CASEditor", "Contextual Animation Editor"));
+	auto WorkspaceMenuCategoryRef = WorkspaceMenuCategory.ToSharedRef();
+
+	FAssetEditorToolkit::RegisterTabSpawners(InTabManager);
+
+	InTabManager->RegisterTabSpawner(FContextualAnimEditorTabs::ViewportID, FOnSpawnTab::CreateSP(this, &FContextualAnimAssetEditorToolkit::SpawnTab_Viewport))
+		.SetDisplayName(LOCTEXT("ViewportTab", "Viewport"))
+		.SetGroup(WorkspaceMenuCategoryRef)
+		.SetIcon(FSlateIcon(FEditorStyle::GetStyleSetName(), "GraphEditor.EventGraph_16x"));
+
+	InTabManager->RegisterTabSpawner(FContextualAnimEditorTabs::AssetDetailsID, FOnSpawnTab::CreateSP(this, &FContextualAnimAssetEditorToolkit::SpawnTab_AssetDetails))
+		.SetDisplayName(LOCTEXT("AssetDetailsTab", "AssetDetails"))
+		.SetGroup(WorkspaceMenuCategoryRef)
+		.SetIcon(FSlateIcon(FEditorStyle::GetStyleSetName(), "LevelEditor.Tabs.Details"));
+
+	InTabManager->RegisterTabSpawner(FContextualAnimEditorTabs::TimelineID, FOnSpawnTab::CreateSP(this, &FContextualAnimAssetEditorToolkit::SpawnTab_Timeline))
+		.SetDisplayName(LOCTEXT("TimelineTab", "Timeline"))
+		.SetGroup(WorkspaceMenuCategoryRef);
+
+ 	InTabManager->RegisterTabSpawner(FContextualAnimEditorTabs::AssetBrowserID, FOnSpawnTab::CreateSP(this, &FContextualAnimAssetEditorToolkit::SpawnTab_AssetBrowser))
+ 		.SetDisplayName(LOCTEXT("AssetBrowserTab", "AssetBrowser"))
+ 		.SetGroup(WorkspaceMenuCategoryRef);
+
+	InTabManager->RegisterTabSpawner(FContextualAnimEditorTabs::PreviewSettingsID, FOnSpawnTab::CreateSP(this, &FContextualAnimAssetEditorToolkit::SpawnTab_PreviewSettings))
+		.SetDisplayName(LOCTEXT("PreviewSceneSettingsTab", "Preview Scene Settings"))
+		.SetGroup(WorkspaceMenuCategoryRef)
+		.SetIcon(FSlateIcon(FEditorStyle::GetStyleSetName(), "LevelEditor.Tabs.Details"));
+}
+
+void FContextualAnimAssetEditorToolkit::UnregisterTabSpawners(const TSharedRef<FTabManager>& InTabManager)
+{
+	FAssetEditorToolkit::UnregisterTabSpawners(InTabManager);
+
+	InTabManager->UnregisterTabSpawner(FContextualAnimEditorTabs::ViewportID);
+	InTabManager->UnregisterTabSpawner(FContextualAnimEditorTabs::AssetDetailsID);
+	InTabManager->UnregisterTabSpawner(FContextualAnimEditorTabs::TimelineID);
+	InTabManager->UnregisterTabSpawner(FContextualAnimEditorTabs::AssetBrowserID);
+	InTabManager->UnregisterTabSpawner(FContextualAnimEditorTabs::PreviewSettingsID);
+}
+
+FName FContextualAnimAssetEditorToolkit::GetToolkitFName() const
+{
+	return FName("ContextualAnimEditor");
+}
+
+FText FContextualAnimAssetEditorToolkit::GetBaseToolkitName() const
+{
+	return LOCTEXT("ContextualAnimEditorAppLabel", "Contextual Anim Editor");
+}
+
+FText FContextualAnimAssetEditorToolkit::GetToolkitName() const
+{
+	FFormatNamedArguments Args;
+	Args.Add(TEXT("AssetName"), FText::FromString(GetSceneAsset()->GetName()));
+	return FText::Format(LOCTEXT("ContextualAnimEditorToolkitName", "{AssetName}"), Args);
+}
+
+FLinearColor FContextualAnimAssetEditorToolkit::GetWorldCentricTabColorScale() const
+{
+	return FLinearColor::White;
+}
+
+FString FContextualAnimAssetEditorToolkit::GetWorldCentricTabPrefix() const
+{
+	return TEXT("ContextualAnimEditor");
+}
+
+TSharedRef<SDockTab> FContextualAnimAssetEditorToolkit::SpawnTab_Viewport(const FSpawnTabArgs& Args)
+{
+	check(Args.GetTabId() == FContextualAnimEditorTabs::ViewportID);
+
+	TSharedRef<SDockTab> SpawnedTab = SNew(SDockTab).Label(LOCTEXT("ViewportTab_Title", "Viewport"));
+
+	if (ViewportWidget.IsValid())
+	{
+		SpawnedTab->SetContent(ViewportWidget.ToSharedRef());
+	}
+
+	return SpawnedTab;
+}
+
+TSharedRef<SDockTab> FContextualAnimAssetEditorToolkit::SpawnTab_Timeline(const FSpawnTabArgs& Args)
+{
+	check(Args.GetTabId().TabType == FContextualAnimEditorTabs::TimelineID);
+
+	TSharedRef<SDockTab> SpawnedTab =
+		SNew(SDockTab)
+		[
+			ViewModel->GetSequencer()->GetSequencerWidget()
+		];
+
+	return SpawnedTab;
+}
+
+TSharedRef<SDockTab> FContextualAnimAssetEditorToolkit::SpawnTab_AssetDetails(const FSpawnTabArgs& Args)
+{
+	check(Args.GetTabId() == FContextualAnimEditorTabs::AssetDetailsID);
+
+	return SNew(SDockTab)
+		.Label(LOCTEXT("AssetDetails_Title", "Asset Details"))
+		[
+			EditingAssetWidget.ToSharedRef()
+		];
+}
+
+TSharedRef<SDockTab> FContextualAnimAssetEditorToolkit::SpawnTab_AssetBrowser(const FSpawnTabArgs& Args)
+{
+	check(Args.GetTabId() == FContextualAnimEditorTabs::AssetBrowserID);
+
+	TSharedRef<SDockTab> SpawnedTab = SNew(SDockTab).Label(LOCTEXT("EditorAssetBrowser_Title", "Asset Browser"));
+
+	if (AssetBrowserWidget.IsValid())
+	{
+		SpawnedTab->SetContent(AssetBrowserWidget.ToSharedRef());
+	}
+
+	return SpawnedTab;
+}
+
+TSharedRef<SDockTab> FContextualAnimAssetEditorToolkit::SpawnTab_PreviewSettings(const FSpawnTabArgs& Args)
+{
+	check(Args.GetTabId() == FContextualAnimEditorTabs::PreviewSettingsID);
+
+	FAdvancedPreviewSceneModule& AdvancedPreviewSceneModule = FModuleManager::LoadModuleChecked<FAdvancedPreviewSceneModule>("AdvancedPreviewScene");
+	TSharedRef<SWidget> InWidget= AdvancedPreviewSceneModule.CreateAdvancedPreviewSceneSettingsWidget(PreviewScene.ToSharedRef());
+
+	TSharedRef<SDockTab> SpawnedTab = SNew(SDockTab)
+		.Label(LOCTEXT("PreviewSceneSettingsTab", "Preview Scene Settings"))
+		[
+			InWidget
+		];
+
+	return SpawnedTab;
+}
+
+void FContextualAnimAssetEditorToolkit::OnFinishedChangingProperties(const FPropertyChangedEvent& PropertyChangedEvent)
+{
+}
+
+#undef LOCTEXT_NAMESPACE
+
diff --git a/Engine/Plugins/Experimental/Animation/ContextualAnimation/Source/ContextualAnimationEditor/Private/ContextualAnimEdMode.cpp b/Engine/Plugins/Experimental/Animation/ContextualAnimation/Source/ContextualAnimationEditor/Private/ContextualAnimEdMode.cpp
--- a/Engine/Plugins/Experimental/Animation/ContextualAnimation/Source/ContextualAnimationEditor/Private/ContextualAnimEdMode.cpp
+++ b/Engine/Plugins/Experimental/Animation/ContextualAnimation/Source/ContextualAnimationEditor/Private/ContextualAnimEdMode.cpp
@@ -1,230 +1,92 @@
 // Copyright Epic Games, Inc. All Rights Reserved.
 
 #include "ContextualAnimEdMode.h"
-#include "ContextualAnimEdModeToolkit.h"
-#include "Toolkits/ToolkitManager.h"
-#include "EditorModeManager.h"
 #include "EditorViewportClient.h"
-#include "Engine/Selection.h"
 #include "EngineUtils.h"
-#include "ContextualAnimEdModeSettings.h"
-#include "DrawDebugHelpers.h"
-#include "ContextualAnimSceneAsset.h"
-#include "ContextualAnimUtilities.h"
-#include "ContextualAnimManager.h"
-#include "ContextualAnimPreviewManager.h"
-#include "ContextualAnimSceneActorComponent.h"
-#include "GameFramework/Character.h"
-#include "ContextualAnimMetadata.h"
-#include "ContextualAnimSceneInstance.h"
-
-const FEditorModeID FContextualAnimEdMode::EM_ContextualAnimEdModeId = TEXT("EM_ContextualAnimEdMode");
+
+const FEditorModeID FContextualAnimEdMode::EdModeId = TEXT("ContextualAnimEdMode");
 
 FContextualAnimEdMode::FContextualAnimEdMode()
 {
-	PreviewManager = NewObject<UContextualAnimPreviewManager>(UContextualAnimPreviewManager::StaticClass());
 }
 
 FContextualAnimEdMode::~FContextualAnimEdMode()
 {
 }
 
-void FContextualAnimEdMode::AddReferencedObjects(FReferenceCollector& Collector)
+bool FContextualAnimEdMode::HandleClick(FEditorViewportClient* InViewportClient, HHitProxy* HitProxy, const FViewportClick& Click)
 {
-	FEdMode::AddReferencedObjects(Collector);
+	bool bHandled = false;
 
-	Collector.AddReferencedObject(PreviewManager);
-}
+	if (HitProxy != nullptr && HitProxy->IsA(HActor::StaticGetType()))
+	{
+		HActor* ActorHitProxy = static_cast<HActor*>(HitProxy);
 
-void FContextualAnimEdMode::Enter()
-{
-	FEdMode::Enter();
+		SelectedActor = ActorHitProxy->Actor;
 
-	if (!Toolkit.IsValid() && UsesToolkits())
-	{
-		Toolkit = MakeShareable(new FContextualAnimEdModeToolkit);
-		Toolkit->Init(Owner->GetToolkitHost());
+		bHandled = true;
 	}
+
+	return bHandled;
 }
 
-void FContextualAnimEdMode::Exit()
+bool FContextualAnimEdMode::GetHitResultUnderCursor(FHitResult& OutHitResult, FEditorViewportClient* InViewportClient, const FViewportClick& Click) const
 {
-	if (Toolkit.IsValid())
-	{
-		FToolkitManager::Get().CloseToolkit(Toolkit.ToSharedRef());
-		Toolkit.Reset();
-	}
+	FSceneViewFamilyContext ViewFamily(FSceneViewFamily::ConstructionValues(InViewportClient->Viewport, InViewportClient->GetScene(), InViewportClient->EngineShowFlags).SetRealtimeUpdate(InViewportClient->IsRealtime()));
+	FSceneView* View = InViewportClient->CalcSceneView(&ViewFamily);
+	FViewportCursorLocation Cursor(View, InViewportClient, Click.GetClickPos().X, Click.GetClickPos().Y);
+	const auto ViewportType = InViewportClient->GetViewportType();
 
-	FEdMode::Exit();
-}
+	const FVector RayStart = Cursor.GetOrigin();
+	const FVector RayEnd = RayStart + Cursor.GetDirection() * HALF_WORLD_MAX;
 
-TSharedPtr<FContextualAnimEdModeToolkit> FContextualAnimEdMode::GetContextualAnimEdModeToolkit() const
-{
-	return StaticCastSharedPtr<FContextualAnimEdModeToolkit>(Toolkit);
+	return InViewportClient->GetWorld()->LineTraceSingleByChannel(OutHitResult, RayStart, RayEnd, ECC_WorldStatic, FCollisionQueryParams::DefaultQueryParam);
 }
 
-void FContextualAnimEdMode::Tick(FEditorViewportClient* ViewportClient, float DeltaTime)
+bool FContextualAnimEdMode::StartTracking(FEditorViewportClient* InViewportClient, FViewport* InViewport)
 {
-	FEdMode::Tick(ViewportClient, DeltaTime);
-
-	if(GEditor->IsSimulatingInEditor())
-	{
-		if (!ViewportClient->Viewport->KeyState(EKeys::RightMouseButton))
-		{
-			if (ViewportClient->Viewport->KeyState(EKeys::W) || ViewportClient->Viewport->KeyState(EKeys::S))
-			{
-				PreviewManager->MoveForward(ViewportClient->Viewport->KeyState(EKeys::W) ? 1.f : -1.f);
-			}
-
-			if (ViewportClient->Viewport->KeyState(EKeys::A) || ViewportClient->Viewport->KeyState(EKeys::D))
-			{
-				PreviewManager->MoveRight(ViewportClient->Viewport->KeyState(EKeys::D) ? 1.f : -1.f);
-			}
-		}
-	}
-	else
-	{
-		if(PreviewManager)
-		{
-			if (PreviewManager->bDrawDebugScene)
-			{
-				if (UContextualAnimSceneAsset* Asset = GetContextualAnimEdModeToolkit()->GetSettings()->SceneAsset)
-				{
-					DrawDebugCoordinateSystem(GetWorld(), PreviewManager->ScenePivot.GetLocation(), PreviewManager->ScenePivot.Rotator(), 50.f, false, 0.f, 0, 1.f);
-					UContextualAnimUtilities::DrawDebugScene(GetWorld(), Asset, PreviewManager->AnimDataIndex, PreviewManager->Time, PreviewManager->ScenePivot, FColor::White, 0.f, 1.f);
-				}
-			}
-
-			//@TODO: Move this to an event that triggers after stop Simulating Mode (not sure if we have one)
-			if(PreviewManager->PreviewActors.Num())
-			{
-				PreviewManager->PreviewActors.Reset();
-			}
-		}
-	}
+	return FEdMode::StartTracking(InViewportClient, InViewport);
 }
 
-bool FContextualAnimEdMode::InputKey(FEditorViewportClient* ViewportClient, FViewport* Viewport, FKey Key, EInputEvent Event)
+bool FContextualAnimEdMode::EndTracking(FEditorViewportClient* InViewportClient, FViewport* InViewport)
 {
-	if(Key == EKeys::Enter && Event == IE_Released && PreviewManager->TestCharacter.IsValid() && GEditor->IsSimulatingInEditor())
-	{
-		UContextualAnimManager* Manager = UContextualAnimManager::Get(GetWorld());
-		check(Manager);		
-
-		if(!Manager->IsActorInAnyScene(PreviewManager->TestCharacter.Get()))
-		{
-			UContextualAnimSceneAsset* Asset = GetContextualAnimEdModeToolkit()->GetSettings()->SceneAsset;
-			if (Asset && PreviewManager->PreviewActors.Num() > 0)
-			{
-				FContextualAnimSceneBindings Bindings;
-				Bindings.RoleToActorMap = PreviewManager->PreviewActors;
-				Manager->TryStartScene(Asset, Bindings);
-			}
-			else
-			{
-				if (UContextualAnimSceneActorComponent* Comp = Manager->FindClosestSceneActorCompToActor(PreviewManager->TestCharacter.Get()))
-				{
-					TSet<UContextualAnimSceneActorComponent*> Comps = Manager->GetSceneActorCompContainer();
-
-					Comps.Sort([Comp](const UContextualAnimSceneActorComponent& A, const UContextualAnimSceneActorComponent& B){
-						const float DistSqA = FVector::DistSquared2D(Comp->GetComponentLocation(), A.GetComponentLocation());
-						const float DistSqB = FVector::DistSquared2D(Comp->GetComponentLocation(), B.GetComponentLocation());
-						return DistSqA < DistSqB;
-					});
-
-					Manager->TryStartScene(Comp->SceneAsset, Comp->GetOwner(), Comps);
-				}
-				else
-				{
-					GEngine->AddOnScreenDebugMessage(-1, 5.f, FColor::Yellow, FString(TEXT("WARNING: Can't find any SceneActor to interact with")));
-					UE_LOG(LogContextualAnim, Warning, TEXT("WARNING: Can't find any SceneActor to interact with"));
-				}
-			}
-		}
-		else
-		{
-			Manager->TryStopSceneWithActor(PreviewManager->TestCharacter.Get());
-		}
-
-		return true;
-	}
-
-	return FEdMode::InputKey(ViewportClient, Viewport, Key, Event);
+	return FEdMode::EndTracking(InViewportClient, InViewport);
 }
 
-bool FContextualAnimEdMode::HandleClick(FEditorViewportClient* InViewportClient, HHitProxy* HitProxy, const FViewportClick& Click)
+bool FContextualAnimEdMode::InputDelta(FEditorViewportClient* InViewportClient, FViewport* InViewport, FVector& InDrag, FRotator& InRot, FVector& InScale)
 {
-	if (!Click.IsAltDown())
-	{
-		return FEdMode::HandleClick(InViewportClient, HitProxy, Click);
-	}
-
-	if (!GEditor->IsSimulatingInEditor())
-	{
-		GEngine->AddOnScreenDebugMessage(-1, 5.f, FColor::Yellow, FString(TEXT("WARNING. You are not in Simulating Mode")));
-		UE_LOG(LogContextualAnim, Warning, TEXT("WARNING. You are not in Simulating Mode"));
-
-		return FEdMode::HandleClick(InViewportClient, HitProxy, Click);
-	}
-
-	FHitResult HitResult;
-	GetHitResultUnderCursor(HitResult, InViewportClient, Click);
+	bool bHandled = false;
 
-	const FString Message = FString::Printf(TEXT("INFO: HandleClick: IsValidHit: %d Actor: %s Location: %s"),
-		HitResult.IsValidBlockingHit(), *GetNameSafe(HitResult.GetActor()), *HitResult.ImpactPoint.ToString());
+	EAxisList::Type CurrentAxis = InViewportClient->GetCurrentWidgetAxis();
 
-	GEngine->AddOnScreenDebugMessage(-1, 5.f, FColor::White, Message);
-	UE_LOG(LogContextualAnim, Warning, TEXT("%s"), *Message);
-
-	if (!HitResult.IsValidBlockingHit())
-	{
-		GEngine->AddOnScreenDebugMessage(-1, 5.f, FColor::Yellow, FString(TEXT("WARNING. HitResult from click event is not a valid blocking hit.")));
-		UE_LOG(LogContextualAnim, Warning, TEXT("WARNING. HitResult from click event is not a valid blocking hit."));
-
-		return FEdMode::HandleClick(InViewportClient, HitProxy, Click);
-	}	
-	
-	FTransform SpawnTransform = FTransform(HitResult.ImpactPoint);
-	UContextualAnimSceneAsset* Asset = GetContextualAnimEdModeToolkit()->GetSettings()->SceneAsset;
-	if (Asset && PreviewManager->PreviewActors.Num() == 0)
-	{
-		PreviewManager->SpawnPreviewActors(Asset, SpawnTransform); 
-	}
-	else
+	if (CurrentAxis != EAxisList::None)
 	{
-		if (!PreviewManager->TestCharacter.IsValid())
+		if (SelectedActor.IsValid())
 		{
-			PreviewManager->TestCharacter = Cast<ACharacter>(PreviewManager->SpawnPreviewActor(PreviewManager->DefaultPreviewClass, SpawnTransform));
-		}
-		else
-		{
-			PreviewManager->MoveToLocation(HitResult.ImpactPoint);
+			SelectedActor->AddActorWorldTransform(FTransform(InRot, InDrag));
+			bHandled = true;
 		}
 	}
 
-	return true;
+	return bHandled;
 }
 
-bool FContextualAnimEdMode::GetHitResultUnderCursor(FHitResult& OutHitResult, FEditorViewportClient* InViewportClient, const FViewportClick& Click) const
+bool FContextualAnimEdMode::AllowWidgetMove()
 {
-	FSceneViewFamilyContext ViewFamily(FSceneViewFamily::ConstructionValues(InViewportClient->Viewport, InViewportClient->GetScene(), InViewportClient->EngineShowFlags).SetRealtimeUpdate(InViewportClient->IsRealtime()));
-	FSceneView* View = InViewportClient->CalcSceneView(&ViewFamily);
-	FViewportCursorLocation Cursor(View, InViewportClient, Click.GetClickPos().X, Click.GetClickPos().Y);
-	const auto ViewportType = InViewportClient->GetViewportType();
-
-	const FVector RayStart = Cursor.GetOrigin();
-	const FVector RayEnd = RayStart + Cursor.GetDirection() * HALF_WORLD_MAX;
-
-	return InViewportClient->GetWorld()->LineTraceSingleByChannel(OutHitResult, RayStart, RayEnd, ECC_WorldStatic, FCollisionQueryParams::DefaultQueryParam);
+	return ShouldDrawWidget();
 }
 
-bool FContextualAnimEdMode::UsesToolkits() const
+bool FContextualAnimEdMode::ShouldDrawWidget() const
 {
-	return true;
+	return SelectedActor.IsValid();
 }
 
-FContextualAnimEdMode& FContextualAnimEdMode::Get()
+FVector FContextualAnimEdMode::GetWidgetLocation() const
 {
-	return *(static_cast<FContextualAnimEdMode*>(GLevelEditorModeTools().GetActiveMode(FContextualAnimEdMode::EM_ContextualAnimEdModeId)));
-}
+	if (SelectedActor.IsValid())
+	{
+		return SelectedActor->GetActorLocation();
+	}
 
+	return FVector::ZeroVector;
+}
diff --git a/Engine/Plugins/Experimental/Animation/ContextualAnimation/Source/ContextualAnimationEditor/Private/ContextualAnimEdModeSettings.cpp b/Engine/Plugins/Experimental/Animation/ContextualAnimation/Source/ContextualAnimationEditor/Private/ContextualAnimEdModeSettings.cpp
deleted file mode 100644
diff --git a/Engine/Plugins/Experimental/Animation/ContextualAnimation/Source/ContextualAnimationEditor/Private/ContextualAnimEdModeToolkit.cpp b/Engine/Plugins/Experimental/Animation/ContextualAnimation/Source/ContextualAnimationEditor/Private/ContextualAnimEdModeToolkit.cpp
deleted file mode 100644
diff --git a/Engine/Plugins/Experimental/Animation/ContextualAnimation/Source/ContextualAnimationEditor/Private/ContextualAnimFactory.cpp b/Engine/Plugins/Experimental/Animation/ContextualAnimation/Source/ContextualAnimationEditor/Private/ContextualAnimFactory.cpp
new file mode 100644
--- /dev/null
+++ b/Engine/Plugins/Experimental/Animation/ContextualAnimation/Source/ContextualAnimationEditor/Private/ContextualAnimFactory.cpp
@@ -0,0 +1,26 @@
+// Copyright Epic Games, Inc. All Rights Reserved.
+
+#include "ContextualAnimFactory.h"
+#include "ContextualAnimSceneAsset.h"
+
+#define LOCTEXT_NAMESPACE "ContextualAnimFactory"
+
+UContextualAnimFactory::UContextualAnimFactory(const FObjectInitializer& ObjectInitializer)
+	: Super(ObjectInitializer)
+{
+	SupportedClass = UContextualAnimSceneAsset::StaticClass();
+	bCreateNew = true;
+	bEditAfterNew = true;
+}
+
+bool UContextualAnimFactory::CanCreateNew() const
+{
+	return true;
+}
+
+UObject* UContextualAnimFactory::FactoryCreateNew(UClass* Class, UObject* InParent, FName Name, EObjectFlags Flags, UObject* Context, FFeedbackContext* Warn)
+{
+	return NewObject<UObject>(InParent, Class, Name, Flags | RF_Transactional);
+}
+
+#undef LOCTEXT_NAMESPACE
diff --git a/Engine/Plugins/Experimental/Animation/ContextualAnimation/Source/ContextualAnimationEditor/Private/ContextualAnimMovieSceneNotifySection.cpp b/Engine/Plugins/Experimental/Animation/ContextualAnimation/Source/ContextualAnimationEditor/Private/ContextualAnimMovieSceneNotifySection.cpp
new file mode 100644
--- /dev/null
+++ b/Engine/Plugins/Experimental/Animation/ContextualAnimation/Source/ContextualAnimationEditor/Private/ContextualAnimMovieSceneNotifySection.cpp
@@ -0,0 +1,46 @@
+// Copyright Epic Games, Inc. All Rights Reserved.
+
+#include "ContextualAnimMovieSceneNotifySection.h"
+#include "ContextualAnimMovieSceneNotifyTrack.h"
+#include "Animation/AnimMontage.h"
+#include "MovieScene.h"
+
+UContextualAnimMovieSceneNotifyTrack* UContextualAnimMovieSceneNotifySection::GetOwnerTrack() const
+{
+	return GetTypedOuter<UContextualAnimMovieSceneNotifyTrack>();
+}
+
+void UContextualAnimMovieSceneNotifySection::Initialize(const FAnimNotifyEvent& NotifyEvent)
+{
+	FFrameRate TickResolution = GetOwnerTrack()->GetTypedOuter<UMovieScene>()->GetTickResolution();
+	FFrameNumber StartFrame = (NotifyEvent.GetTriggerTime() * TickResolution).RoundToFrame();
+	FFrameNumber EndFrame = (NotifyEvent.GetEndTriggerTime() * TickResolution).RoundToFrame();
+	SetRange(TRange<FFrameNumber>::Exclusive(StartFrame, EndFrame));
+
+	AnimNotifyEventGuid = NotifyEvent.Guid;
+}
+
+FAnimNotifyEvent* UContextualAnimMovieSceneNotifySection::GetAnimNotifyEvent() const
+{
+	UAnimMontage& Animation = GetOwnerTrack()->GetAnimation();
+	for (FAnimNotifyEvent& NotifyEvent : Animation.Notifies)
+	{
+		if (NotifyEvent.Guid == AnimNotifyEventGuid)
+		{
+			return &NotifyEvent;
+		}
+	}
+
+	return nullptr;
+}
+
+UAnimNotifyState* UContextualAnimMovieSceneNotifySection::GetAnimNotifyState() const
+{
+	UAnimMontage& Animation = GetOwnerTrack()->GetAnimation();
+	if(const FAnimNotifyEvent* NotifyEventPtr = GetAnimNotifyEvent())
+	{
+		return NotifyEventPtr->NotifyStateClass;
+	}
+
+	return nullptr;
+}
\ No newline at end of file
diff --git a/Engine/Plugins/Experimental/Animation/ContextualAnimation/Source/ContextualAnimationEditor/Private/ContextualAnimMovieSceneNotifyTrack.cpp b/Engine/Plugins/Experimental/Animation/ContextualAnimation/Source/ContextualAnimationEditor/Private/ContextualAnimMovieSceneNotifyTrack.cpp
new file mode 100644
--- /dev/null
+++ b/Engine/Plugins/Experimental/Animation/ContextualAnimation/Source/ContextualAnimationEditor/Private/ContextualAnimMovieSceneNotifyTrack.cpp
@@ -0,0 +1,141 @@
+// Copyright Epic Games, Inc. All Rights Reserved.
+
+#include "ContextualAnimMovieSceneNotifyTrack.h"
+#include "ContextualAnimMovieSceneNotifySection.h"
+#include "Animation/AnimMontage.h"
+#include "Animation/AnimNotifies/AnimNotify.h"
+#include "Animation/AnimNotifies/AnimNotifyState.h"
+#include "ContextualAnimViewModel.h"
+#include "ContextualAnimMovieSceneSequence.h"
+#include "Subsystems/AssetEditorSubsystem.h"
+#include "Editor.h"
+
+FContextualAnimViewModel& UContextualAnimMovieSceneNotifyTrack::GetViewModel() const
+{
+	return GetTypedOuter<UMovieScene>()->GetTypedOuter<UContextualAnimMovieSceneSequence>()->GetViewModel();
+}
+
+void UContextualAnimMovieSceneNotifyTrack::Initialize(UAnimMontage& AnimationRef, const FAnimNotifyTrack& NotifyTrack)
+{
+	// Cache animation this track was created from
+	Animation = &AnimationRef;
+
+	SetDisplayName(FText::FromName(NotifyTrack.TrackName));
+
+	for (const FAnimNotifyEvent* NotifyEvent : NotifyTrack.Notifies)
+	{
+		if (UAnimNotifyState* Notify = Cast<UAnimNotifyState>(NotifyEvent->NotifyStateClass))
+		{
+			// Create new movie scene section
+			UContextualAnimMovieSceneNotifySection* Section = CastChecked<UContextualAnimMovieSceneNotifySection>(CreateNewSection());
+
+			// Set range and cache guid
+			Section->Initialize(*NotifyEvent);
+
+			// Add section to the track
+			AddSection(*Section);
+		}
+	}
+}
+
+UAnimMontage& UContextualAnimMovieSceneNotifyTrack::GetAnimation() const 
+{ 
+	check(Animation.IsValid());
+	return *Animation.Get();
+}
+
+UMovieSceneSection* UContextualAnimMovieSceneNotifyTrack::CreateNewSection()
+{
+	return NewObject<UContextualAnimMovieSceneNotifySection>(this, NAME_None, RF_Transactional);
+}
+
+void UContextualAnimMovieSceneNotifyTrack::AddSection(UMovieSceneSection& Section)
+{
+	Sections.Add(&Section);
+}
+
+bool UContextualAnimMovieSceneNotifyTrack::SupportsType(TSubclassOf<UMovieSceneSection> SectionClass) const
+{
+	return SectionClass == UContextualAnimMovieSceneNotifySection::StaticClass();;
+}
+
+const TArray<UMovieSceneSection*>& UContextualAnimMovieSceneNotifyTrack::GetAllSections() const
+{
+	return Sections;
+}
+
+bool UContextualAnimMovieSceneNotifyTrack::HasSection(const UMovieSceneSection& Section) const
+{
+	return Sections.Contains(&Section);
+}
+
+bool UContextualAnimMovieSceneNotifyTrack::IsEmpty() const
+{
+	return Sections.Num() != 0;
+}
+
+void UContextualAnimMovieSceneNotifyTrack::RemoveSection(UMovieSceneSection& Section)
+{
+	// Temp solution to prevent editor from crashing when removing the notify from here while the animation is open in the anim editor
+	UAssetEditorSubsystem* AssetEditorSubsystem = GEditor->GetEditorSubsystem<UAssetEditorSubsystem>();
+	check(AssetEditorSubsystem);
+	AssetEditorSubsystem->CloseAllEditorsForAsset(&GetAnimation());
+
+	// Remove Notify from the animation
+	UContextualAnimMovieSceneNotifySection* NotifySection = CastChecked<UContextualAnimMovieSceneNotifySection>(&Section);
+
+	UAnimMontage& Anim = GetAnimation();
+
+	FGuid Guid = NotifySection->GetAnimNotifyEventGuid();
+	const int32 TotalRemoved = Anim.Notifies.RemoveAll([&Guid](const FAnimNotifyEvent& Event) { return Event.Guid == Guid; });
+	if (TotalRemoved)
+	{
+		GetViewModel().AnimationModified(Anim);
+	}
+
+	// Remove movie scene section
+	Sections.Remove(&Section);
+}
+
+void UContextualAnimMovieSceneNotifyTrack::RemoveSectionAt(int32 SectionIndex)
+{
+	check(Sections.IsValidIndex(SectionIndex));
+	RemoveSection(*Sections[SectionIndex]);
+}
+
+#if WITH_EDITOR
+EMovieSceneSectionMovedResult UContextualAnimMovieSceneNotifyTrack::OnSectionMoved(UMovieSceneSection& Section, const FMovieSceneSectionMovedParams& Params)
+{
+	if (Params.MoveType == EPropertyChangeType::ValueSet)
+	{
+		UContextualAnimMovieSceneNotifySection* NotifySection = CastChecked<UContextualAnimMovieSceneNotifySection>(&Section);
+
+		FFrameRate TickResolution = GetTypedOuter<UMovieScene>()->GetTickResolution();
+
+		if (FAnimNotifyEvent* NotifyEventPtr = NotifySection->GetAnimNotifyEvent())
+		{
+			UAnimMontage* Anim = nullptr;
+			if (NotifyEventPtr->NotifyStateClass)
+			{
+				Anim = NotifyEventPtr->NotifyStateClass->GetTypedOuter<UAnimMontage>();
+			}
+			else if (NotifyEventPtr->Notify)
+			{
+				Anim = NotifyEventPtr->Notify->GetTypedOuter<UAnimMontage>();
+			}
+
+			check(&GetAnimation() == Anim);
+
+			const float SectionStartTime = (float)TickResolution.AsSeconds(NotifySection->GetInclusiveStartFrame());
+			const float SectionEndTime = (float)TickResolution.AsSeconds(NotifySection->GetExclusiveEndFrame());
+
+			NotifyEventPtr->SetTime(SectionStartTime);
+			NotifyEventPtr->SetDuration(SectionEndTime - SectionStartTime);
+
+			GetViewModel().AnimationModified(*Anim);
+		}
+	}
+
+	return EMovieSceneSectionMovedResult::None;
+}
+#endif
\ No newline at end of file
diff --git a/Engine/Plugins/Experimental/Animation/ContextualAnimation/Source/ContextualAnimationEditor/Private/ContextualAnimMovieSceneNotifyTrackEditor.cpp b/Engine/Plugins/Experimental/Animation/ContextualAnimation/Source/ContextualAnimationEditor/Private/ContextualAnimMovieSceneNotifyTrackEditor.cpp
new file mode 100644
--- /dev/null
+++ b/Engine/Plugins/Experimental/Animation/ContextualAnimation/Source/ContextualAnimationEditor/Private/ContextualAnimMovieSceneNotifyTrackEditor.cpp
@@ -0,0 +1,614 @@
+// Copyright Epic Games, Inc. All Rights Reserved.
+
+#include "ContextualAnimMovieSceneNotifyTrackEditor.h"
+#include "ContextualAnimMovieSceneNotifyTrack.h"
+#include "ContextualAnimMovieSceneNotifySection.h"
+#include "ContextualAnimSceneAsset.h"
+#include "ContextualAnimViewModel.h"
+#include "ContextualAnimMovieSceneSequence.h"
+#include "SequencerUtilities.h"
+#include "SequencerSectionPainter.h"
+#include "ClassViewerModule.h"
+#include "ClassViewerFilter.h"
+#include "Animation/AnimNotifies/AnimNotifyState.h"
+#include "Animation/AnimNotifies/AnimNotify.h"
+#include "Animation/AnimMontage.h"
+#include "AnimNotifyState_IKWindow.h"
+#include "IStructureDetailsView.h"
+#include "PropertyEditorModule.h"
+#include "PropertyCustomizationHelpers.h"
+#include "PersonaUtils.h"
+#include "Framework/MultiBox/MultiBoxBuilder.h"
+#include "Modules/ModuleManager.h"
+#include "Widgets/Input/SButton.h"
+
+#define LOCTEXT_NAMESPACE "FContextualAnimMovieSceneNotifyTrackEditor"
+
+// FNewRoleWidgetParams
+////////////////////////////////////////////////////////////////////////////////////////////////
+
+bool FNewRoleWidgetParams::HasValidData() const
+{
+	return RoleName != NAME_None && PreviewClass && Animation;
+}
+
+// UNewIKTargetWidgetParams
+////////////////////////////////////////////////////////////////////////////////////////////////
+
+void UNewIKTargetWidgetParams::Reset(const FName& InSourceRole, const UContextualAnimSceneAsset& InSceneAsset)
+{
+	CachedRoles = InSceneAsset.GetRoles();
+	check(CachedRoles.Contains(InSourceRole));
+
+	SceneAssetPtr = &InSceneAsset;
+
+	Params.SourceRole = InSourceRole;
+	Params.GoalName = NAME_None;
+	Params.TargetBone = FBoneReference();
+	Params.SourceBone = FBoneReference();
+
+	if(CachedRoles.Num() > 1)
+	{
+		Params.TargetRole = *CachedRoles.FindByPredicate([this](const FName& Role){ return Role != Params.SourceRole; });
+	}
+}
+
+bool UNewIKTargetWidgetParams::HasValidData() const
+{
+	return	Params.GoalName != NAME_None &&
+			Params.TargetBone.BoneName != NAME_None &&
+			Params.SourceBone.BoneName != NAME_None &&
+			CachedRoles.Contains(Params.TargetRole);
+}
+
+const UContextualAnimSceneAsset& UNewIKTargetWidgetParams::GetSceneAsset() const
+{
+	check(SceneAssetPtr.IsValid());
+	return *SceneAssetPtr.Get();
+}
+
+USkeleton* UNewIKTargetWidgetParams::GetSkeleton(bool& bInvalidSkeletonIsError, const IPropertyHandle* PropertyHandle)
+{
+	//bInvalidSkeletonIsError = false;
+
+	if(PropertyHandle)
+	{
+		if(PropertyHandle->GetProperty()->GetFName() == GET_MEMBER_NAME_CHECKED(FNewIKTargetParams, SourceBone))
+		{
+			UAnimMontage* Animation = GetSceneAsset().GetAnimationForRoleAtIndex(Params.SourceRole, 0);
+			check(Animation);
+
+			return Animation->GetSkeleton();
+		}
+		else if(PropertyHandle->GetProperty()->GetFName() == GET_MEMBER_NAME_CHECKED(FNewIKTargetParams, TargetBone))
+		{
+			UAnimMontage* Animation = GetSceneAsset().GetAnimationForRoleAtIndex(Params.TargetRole, 0);
+			check(Animation);
+
+			return Animation->GetSkeleton();
+		}
+	}
+
+	return nullptr;
+}
+
+TArray<FString> UNewIKTargetWidgetParams::GetTargetRoleOptions() const
+{
+	TArray<FString> Options;
+
+	for(const FName& Role : CachedRoles)
+	{
+		if(Role != Params.SourceRole)
+		{
+			Options.Add(Role.ToString());
+		}
+	}
+	
+	return Options;
+}
+
+// FContextualAnimMovieSceneNotifyTrackEditor
+////////////////////////////////////////////////////////////////////////////////////////////////
+
+FContextualAnimMovieSceneNotifyTrackEditor::FContextualAnimMovieSceneNotifyTrackEditor(TSharedRef<ISequencer> InSequencer)
+	: FMovieSceneTrackEditor(InSequencer)
+	, NewIKTargetWidgetParams(nullptr)
+{ 
+}
+
+void FContextualAnimMovieSceneNotifyTrackEditor::AddReferencedObjects(FReferenceCollector& Collector)
+{
+	if (NewIKTargetWidgetParams)
+	{
+		Collector.AddReferencedObject(NewIKTargetWidgetParams);
+	}
+}
+
+UContextualAnimMovieSceneSequence& FContextualAnimMovieSceneNotifyTrackEditor::GetMovieSceneSequence() const
+{
+	UContextualAnimMovieSceneSequence* Sequence = GetFocusedMovieScene()->GetTypedOuter<UContextualAnimMovieSceneSequence>();
+	check(Sequence);
+	return *Sequence;
+}
+
+TSharedRef<ISequencerTrackEditor> FContextualAnimMovieSceneNotifyTrackEditor::CreateTrackEditor(TSharedRef<ISequencer> InSequencer)
+{
+	return MakeShareable(new FContextualAnimMovieSceneNotifyTrackEditor(InSequencer));
+}
+
+TSharedRef<ISequencerSection> FContextualAnimMovieSceneNotifyTrackEditor::MakeSectionInterface(UMovieSceneSection& SectionObject, UMovieSceneTrack& Track, FGuid ObjectBinding)
+{
+	checkf(SectionObject.GetClass()->IsChildOf<UContextualAnimMovieSceneNotifySection>(), TEXT("Unsupported section."));
+	return MakeShareable(new FContextualAnimNotifySection(SectionObject));
+}
+
+void FContextualAnimMovieSceneNotifyTrackEditor::FillNewNotifyStateMenu(FMenuBuilder& MenuBuilder, bool bIsReplaceWithMenu, UContextualAnimMovieSceneNotifyTrack* Track, int32 RowIndex)
+{
+	// MenuBuilder always has a search widget added to it by default, hence if larger then 1 then something else has been added to it
+	if (MenuBuilder.GetMultiBox()->GetBlocks().Num() > 1)
+	{
+		MenuBuilder.AddMenuSeparator();
+	}
+
+	FOnClassPicked OnPicked(FOnClassPicked::CreateLambda([this, bIsReplaceWithMenu, Track, RowIndex](UClass* InClass)
+		{
+			//FSlateApplication::Get().DismissAllMenus();
+
+			if (bIsReplaceWithMenu)
+			{
+				//@TODO: Implement
+			}
+			else
+			{
+				CreateNewSection(Track, RowIndex, InClass);
+			}
+		}
+	));
+
+	TSharedRef<SWidget> AnimNotifyStateClassPicker = PersonaUtils::MakeAnimNotifyStatePicker(&Track->GetAnimation(), OnPicked);
+	MenuBuilder.AddWidget(AnimNotifyStateClassPicker, FText(), true, false);
+}
+
+void FContextualAnimMovieSceneNotifyTrackEditor::BuildAddTrackMenu(FMenuBuilder& MenuBuilder)
+{
+	// Menu that appears when clicking on the Add Track button next to the Search Tracks bar
+
+	// @TODO: Temporally using this TrackEditor to add New Actor Tracks (new roles to the interaction)
+
+	MenuBuilder.AddSubMenu(
+		LOCTEXT("AddActorTrack", "Actor Track"),
+		LOCTEXT("AddActorTrackTooltip", "Adds new actor track"),
+		FNewMenuDelegate::CreateRaw(this, &FContextualAnimMovieSceneNotifyTrackEditor::BuildAddTrackSubMenu, TArray<FGuid>()),
+		false,
+		FSlateIcon()
+	);
+}
+
+void FContextualAnimMovieSceneNotifyTrackEditor::BuildAddTrackSubMenu(FMenuBuilder& MenuBuilder, TArray<FGuid> ObjectBindings)
+{
+	FPropertyEditorModule& PropertyModule = FModuleManager::LoadModuleChecked<FPropertyEditorModule>("PropertyEditor");
+
+	FDetailsViewArgs Args;
+	Args.bHideSelectionTip = true;
+	Args.bAllowSearch = false;
+	Args.bAllowFavoriteSystem = false;
+
+	NewRoleWidgetParams = MakeShared<FStructOnScope>(FNewRoleWidgetParams::StaticStruct());
+
+	TSharedRef<IStructureDetailsView> StructureDetailsView = PropertyModule.CreateStructureDetailView(Args, FStructureDetailsViewArgs(), NewRoleWidgetParams);
+	
+	MenuBuilder.AddWidget(
+		SNew(SBox)
+		.MinDesiredWidth(500.0f)
+		.MaxDesiredWidth(500.f)
+		.MaxDesiredHeight(400.0f)
+		[
+			SNew(SVerticalBox)
+			+ SVerticalBox::Slot()
+			.AutoHeight()
+			.HAlign(HAlign_Fill)
+			[
+				StructureDetailsView->GetWidget().ToSharedRef()
+			]
+			+SVerticalBox::Slot()
+			.AutoHeight()
+			.HAlign(HAlign_Fill)
+			[
+				SNew(SButton)
+				.ContentPadding(3)
+				.VAlign(VAlign_Center)
+				.HAlign(HAlign_Center)
+				.IsEnabled_Lambda([this]() { return ((FNewRoleWidgetParams*)NewRoleWidgetParams->GetStructMemory())->HasValidData(); })
+				.OnClicked_Lambda([this]()
+					{
+						FNewRoleWidgetParams* Data = (FNewRoleWidgetParams*)NewRoleWidgetParams->GetStructMemory();
+
+						GetMovieSceneSequence().GetViewModel().AddActorTrack(*Data);
+
+						FSlateApplication::Get().DismissAllMenus();
+
+						return FReply::Handled();
+					})
+				.Text(LOCTEXT("OK", "OK"))
+			]
+		],
+		FText(), true, false);
+}
+
+TSharedPtr<SWidget> FContextualAnimMovieSceneNotifyTrackEditor::BuildOutlinerEditWidget(const FGuid& ObjectBinding, UMovieSceneTrack* Track, const FBuildEditWidgetParams& Params)
+{
+	// [+Section] button on the track
+
+	check(Track);
+
+	TSharedPtr<ISequencer> SequencerPtr = GetSequencer();
+	if (!SequencerPtr.IsValid())
+	{
+		return SNullWidget::NullWidget;
+	}
+
+	TWeakObjectPtr<UContextualAnimMovieSceneNotifyTrack> WeakTrack = Cast<UContextualAnimMovieSceneNotifyTrack>(Track);
+	const int32 RowIndex = Params.TrackInsertRowIndex;
+	auto SubMenuCallback = [this, WeakTrack, RowIndex]
+	{
+		FMenuBuilder MenuBuilder(true, nullptr);
+
+		UContextualAnimMovieSceneNotifyTrack* TrackPtr = WeakTrack.Get();
+		if (TrackPtr)
+		{
+			//@TODO: Add shortcut for Motion Warping window too
+
+			MenuBuilder.AddSubMenu(
+				LOCTEXT("AddIKWindow", "Add IK Window"),
+				LOCTEXT("AddIKWindowTooltip", "Adds new IK Window"),
+				FNewMenuDelegate::CreateRaw(this, &FContextualAnimMovieSceneNotifyTrackEditor::BuildNewIKTargetSubMenu, TrackPtr, RowIndex));
+
+
+			MenuBuilder.AddSubMenu(
+				LOCTEXT("AddNotifyState", "Add Notify State"),
+				LOCTEXT("AddNotifyStateToolTip", "Adds new AnimNotifyState"),
+				FNewMenuDelegate::CreateRaw(this, &FContextualAnimMovieSceneNotifyTrackEditor::FillNewNotifyStateMenu, false, TrackPtr, RowIndex));
+		}
+		else
+		{
+			MenuBuilder.AddWidget(SNew(STextBlock).Text(LOCTEXT("InvalidTrack", "Track is no longer valid")), FText(), true);
+		}
+
+		return MenuBuilder.MakeWidget();
+	};
+
+	return SNew(SHorizontalBox)
+		+ SHorizontalBox::Slot()
+		.AutoWidth()
+		.VAlign(VAlign_Center)
+		[
+			FSequencerUtilities::MakeAddButton(LOCTEXT("AddSection", "Section"), FOnGetContent::CreateLambda(SubMenuCallback), Params.NodeIsHovered, GetSequencer())
+		];
+}
+
+void FContextualAnimMovieSceneNotifyTrackEditor::BuildNewIKTargetSubMenu(FMenuBuilder& MenuBuilder, UContextualAnimMovieSceneNotifyTrack* Track, int32 RowIndex)
+{
+	// Create new IK Target
+	{
+		MenuBuilder.BeginSection(NAME_None, LOCTEXT("CreateNewIKTarget", "Create new IK Target"));
+
+		MenuBuilder.AddSubMenu(
+			LOCTEXT("NewIKTarget", "New IK Target"),
+			LOCTEXT("NewIKTargetTooltip", "Creates a new IK Target and adds an IK window for it"),
+			FNewMenuDelegate::CreateRaw(this, &FContextualAnimMovieSceneNotifyTrackEditor::BuildNewIKTargetWidget, Track, RowIndex),
+			false,
+			FSlateIcon()
+		);
+
+		MenuBuilder.EndSection();
+	}
+
+	// List of IK Targets already created
+	{
+		MenuBuilder.BeginSection(NAME_None, LOCTEXT("AvailableIKTargets", "Available IK Targets"));
+
+		const UContextualAnimSceneAsset* SceneAsset = GetMovieSceneSequence().GetViewModel().GetSceneAsset();
+		const FName Role = SceneAsset->FindRoleByAnimation(&Track->GetAnimation());
+		check(Role != NAME_None);
+
+		TArray<FName> ExistingGoals = SceneAsset->GetIKGoalsForRole(Role);
+		for(const FName& GoalName : ExistingGoals)
+		{
+			MenuBuilder.AddMenuEntry(
+ 			FText::FromName(GoalName),
+			FText::GetEmpty(),
+ 			FSlateIcon(),
+			FUIAction(FExecuteAction::CreateLambda([this, Track, RowIndex, GoalName]() {
+					CreateNewIKSection(Track, RowIndex, GoalName);
+			})));
+		}
+
+		MenuBuilder.EndSection();
+	}
+}
+
+UContextualAnimMovieSceneNotifySection* FContextualAnimMovieSceneNotifyTrackEditor::CreateNewIKSection(UContextualAnimMovieSceneNotifyTrack* Track, int32 RowIndex, const FName& GoalName)
+{
+	UContextualAnimMovieSceneNotifySection* NewSection = CreateNewSection(Track, RowIndex, UAnimNotifyState_IKWindow::StaticClass());
+	check(NewSection);
+
+	UAnimNotifyState_IKWindow* AnimNotify = Cast<UAnimNotifyState_IKWindow>(NewSection->GetAnimNotifyState());
+	check(AnimNotify);
+ 
+	AnimNotify->GoalName = GoalName;
+
+	return NewSection;
+}
+
+void FContextualAnimMovieSceneNotifyTrackEditor::BuildNewIKTargetWidget(FMenuBuilder& MenuBuilder, UContextualAnimMovieSceneNotifyTrack* Track, int32 RowIndex)
+{
+	if(NewIKTargetWidgetParams == nullptr)
+	{
+		NewIKTargetWidgetParams = NewObject<UNewIKTargetWidgetParams>(UNewIKTargetWidgetParams::StaticClass());
+	}
+	
+	check(NewIKTargetWidgetParams);
+
+	const UContextualAnimSceneAsset* SceneAsset = GetMovieSceneSequence().GetViewModel().GetSceneAsset();
+	const FName SourceRole = SceneAsset->FindRoleByAnimation(&Track->GetAnimation());
+	check(SourceRole != NAME_None);
+
+	NewIKTargetWidgetParams->Reset(SourceRole, *SceneAsset);
+
+	FPropertyEditorModule& PropertyModule = FModuleManager::LoadModuleChecked<FPropertyEditorModule>("PropertyEditor");
+
+	FDetailsViewArgs Args;
+	Args.bHideSelectionTip = true;
+	Args.bAllowSearch = false;
+	Args.bAllowFavoriteSystem = false;
+
+	TSharedPtr<IDetailsView> Widget = PropertyModule.CreateDetailView(Args);
+	Widget->SetObject(NewIKTargetWidgetParams);
+
+	MenuBuilder.AddWidget(
+		SNew(SBox)
+		.MinDesiredWidth(500.0f)
+		.MaxDesiredWidth(500.f)
+		.MaxDesiredHeight(400.0f)
+		[
+			SNew(SVerticalBox)
+			+ SVerticalBox::Slot()
+			.AutoHeight()
+			.HAlign(HAlign_Fill)
+			[
+				Widget.ToSharedRef()
+			]
+			+ SVerticalBox::Slot()
+			.AutoHeight()
+			.HAlign(HAlign_Fill)
+			[
+				SNew(SButton)
+				.ContentPadding(3)
+				.VAlign(VAlign_Center)
+				.HAlign(HAlign_Center)
+				.IsEnabled_Lambda([this](){ return NewIKTargetWidgetParams->HasValidData(); })
+				.OnClicked_Lambda([this, Track, RowIndex]()
+				{
+					UContextualAnimSceneAsset* SceneAsset = GetMovieSceneSequence().GetViewModel().GetSceneAsset();
+					check(SceneAsset);
+
+					FContextualAnimCompositeTrack* SceneAssetTrack = SceneAsset->DataContainer.Find(NewIKTargetWidgetParams->Params.SourceRole);
+					check(SceneAssetTrack);
+
+					// Add IK Target definition to the scene asset
+					FContextualAnimIKTargetDefinition IKTargetDef;
+					IKTargetDef.GoalName = NewIKTargetWidgetParams->Params.GoalName;
+					IKTargetDef.BoneName = NewIKTargetWidgetParams->Params.SourceBone.BoneName;
+					IKTargetDef.Provider = NewIKTargetWidgetParams->Params.Provider;
+					IKTargetDef.TargetRoleName = NewIKTargetWidgetParams->Params.TargetRole;
+					IKTargetDef.TargetBoneName = NewIKTargetWidgetParams->Params.TargetBone.BoneName;
+
+					SceneAssetTrack->Settings.IKTargetDefinitions.AddUnique(IKTargetDef);
+
+					SceneAsset->PrecomputeData();
+
+					SceneAsset->MarkPackageDirty();
+
+					// Create IK section for the newly created target
+	 				CreateNewIKSection(Track, RowIndex, IKTargetDef.GoalName);
+
+					FSlateApplication::Get().DismissAllMenus();
+
+					return FReply::Handled();
+				})
+			.Text(LOCTEXT("OK", "OK"))
+			]
+		],
+		FText(), true, false);
+}
+
+void FContextualAnimMovieSceneNotifyTrackEditor::BuildObjectBindingTrackMenu(FMenuBuilder& MenuBuilder, const TArray<FGuid>& ObjectBindings, const UClass* ObjectClass)
+{	
+	// Builds menu that appears when clicking on the +Track button on an Object Track
+
+	UMovieSceneSequence* MovieSequence = GetSequencer()->GetFocusedMovieSceneSequence();
+	if(MovieSequence && MovieSequence->GetMovieScene()->FindPossessable(ObjectBindings[0]))
+	{
+		MenuBuilder.AddMenuEntry(
+			LOCTEXT("NotifyTrack", "Notify Track"),
+			LOCTEXT("NotifyTrackTooltip", "Adds a notify track"),
+			FSlateIcon(),
+			FUIAction(FExecuteAction::CreateRaw(this, &FContextualAnimMovieSceneNotifyTrackEditor::AddNewNotifyTrack, ObjectBindings)
+			)
+		);
+	}
+}
+
+void FContextualAnimMovieSceneNotifyTrackEditor::AddNewNotifyTrack(TArray<FGuid> ObjectBindings)
+{
+	// Copied from AnimTimelineTrack_Notifies.cpp/FAnimTimelineTrack_Notifies::GetNewTrackName()
+	auto GetNewTrackName = [](UAnimSequenceBase* InAnimSequenceBase) -> FName
+	{
+		TArray<FName> TrackNames;
+		TrackNames.Reserve(50);
+
+		for (const FAnimNotifyTrack& Track : InAnimSequenceBase->AnimNotifyTracks)
+		{
+			TrackNames.Add(Track.TrackName);
+		}
+
+		FName NameToTest;
+		int32 TrackIndex = 1;
+
+		do
+		{
+			NameToTest = *FString::FromInt(TrackIndex++);
+		} while (TrackNames.Contains(NameToTest));
+
+		return NameToTest;
+	};
+
+	for (const FGuid& ObjectBinding : ObjectBindings)
+	{
+		UAnimMontage* Animation = GetMovieSceneSequence().GetViewModel().FindAnimationByGuid(ObjectBinding);
+		check(Animation);
+
+		// Copied from AnimTimelineTrack_Notifies.cpp/FAnimTimelineTrack_Notifies::AddTrack()
+
+		FAnimNotifyTrack NewItem;
+		NewItem.TrackName = GetNewTrackName(Animation);
+		NewItem.TrackColor = FLinearColor::White;
+
+		Animation->AnimNotifyTracks.Add(NewItem);
+
+		GetMovieSceneSequence().GetViewModel().AnimationModified(*Animation);
+
+		// Create MovieSceneTrack
+		UContextualAnimMovieSceneNotifyTrack* MovieSceneTrack = GetMovieSceneSequence().GetMovieScene()->AddTrack<UContextualAnimMovieSceneNotifyTrack>(ObjectBinding);
+		check(MovieSceneTrack);
+
+		MovieSceneTrack->SetDisplayName(FText::FromString(NewItem.TrackName.ToString()));
+	}
+
+	GetSequencer()->NotifyMovieSceneDataChanged(EMovieSceneDataChangeType::MovieSceneStructureItemAdded);
+}
+
+bool FContextualAnimMovieSceneNotifyTrackEditor::SupportsType(TSubclassOf<UMovieSceneTrack> Type) const
+{
+	return (Type == UContextualAnimMovieSceneNotifyTrack::StaticClass());
+}
+
+bool FContextualAnimMovieSceneNotifyTrackEditor::SupportsSequence(UMovieSceneSequence* InSequence) const
+{
+	ETrackSupport TrackSupported = InSequence ? InSequence->IsTrackSupported(UContextualAnimMovieSceneNotifyTrack::StaticClass()) : ETrackSupport::NotSupported;
+	return TrackSupported == ETrackSupport::Supported;
+}
+
+UContextualAnimMovieSceneNotifySection* FContextualAnimMovieSceneNotifyTrackEditor::CreateNewSection(UContextualAnimMovieSceneNotifyTrack* Track, int32 RowIndex, UClass* NotifyClass)
+{
+	TSharedPtr<ISequencer> SequencerPtr = GetSequencer();
+	if (SequencerPtr.IsValid())
+	{
+		UAnimMontage& Animation = Track->GetAnimation();
+
+		const FName TrackName = FName(Track->GetDisplayName().ToString());
+		const int32 TrackIndex = Animation.AnimNotifyTracks.IndexOfByPredicate([TrackName](const FAnimNotifyTrack& AnimNotifyTrack) { return AnimNotifyTrack.TrackName == TrackName; });
+		check(TrackIndex != INDEX_NONE);
+
+		// Add new AnimNotifyEvent to the animation
+		FAnimNotifyEvent& NewEvent = Animation.Notifies.AddDefaulted_GetRef();
+		{
+			NewEvent.Guid = FGuid::NewGuid();
+
+			const float StartTime = 0.f;
+			NewEvent.Link(&Animation, StartTime);
+			NewEvent.TriggerTimeOffset = GetTriggerTimeOffsetForType(Animation.CalculateOffsetForNotify(StartTime));
+			NewEvent.TrackIndex = TrackIndex;
+
+			if (NotifyClass)
+			{
+				UObject* AnimNotifyClass = NewObject<UObject>(&Animation, NotifyClass, NAME_None, RF_Transactional);
+				NewEvent.NotifyStateClass = Cast<UAnimNotifyState>(AnimNotifyClass);
+				NewEvent.Notify = Cast<UAnimNotify>(AnimNotifyClass);
+
+				if (NewEvent.NotifyStateClass)
+				{
+					// Set default duration to 1 frame for AnimNotifyState.
+					NewEvent.SetDuration(1 / 30.f);
+					NewEvent.EndLink.Link(&Animation, NewEvent.EndLink.GetTime());
+					NewEvent.TriggerWeightThreshold = NewEvent.NotifyStateClass->GetDefaultTriggerWeightThreshold();
+					NewEvent.NotifyName = FName(*NewEvent.NotifyStateClass->GetNotifyName());
+				}
+				else if (NewEvent.Notify)
+				{
+					NewEvent.TriggerWeightThreshold = NewEvent.Notify->GetDefaultTriggerWeightThreshold();
+					NewEvent.NotifyName = FName(*NewEvent.Notify->GetNotifyName());
+				}
+			}
+			else
+			{
+				NewEvent.Notify = nullptr;
+				NewEvent.NotifyStateClass = nullptr;
+			}
+		}
+
+		// Create new movie scene section
+		UContextualAnimMovieSceneNotifySection* NewSection = CastChecked<UContextualAnimMovieSceneNotifySection>(Track->CreateNewSection());
+
+		NewSection->SetRowIndex(RowIndex);
+
+		// Set range and cache guid
+		NewSection->Initialize(NewEvent);
+	
+		GetMovieSceneSequence().GetViewModel().AnimationModified(Animation);
+
+		// Add section to the track
+		Track->AddSection(*NewSection);
+		Track->UpdateEasing();
+		Track->Modify();
+
+		// Select new section
+		SequencerPtr->EmptySelection();
+		SequencerPtr->SelectSection(NewSection);
+		SequencerPtr->ThrobSectionSelection();
+
+		SequencerPtr->NotifyMovieSceneDataChanged(EMovieSceneDataChangeType::MovieSceneStructureItemAdded);
+
+		return NewSection;
+	}
+
+	return nullptr;
+}
+
+// FContextualAnimNotifySection
+////////////////////////////////////////////////////////////////////////////////////////////////
+
+FContextualAnimNotifySection::FContextualAnimNotifySection(UMovieSceneSection& InSection)
+	: Section(*CastChecked<UContextualAnimMovieSceneNotifySection>(&InSection))
+{
+}
+
+UMovieSceneSection* FContextualAnimNotifySection::GetSectionObject()
+{
+	return &Section;
+}
+
+int32 FContextualAnimNotifySection::OnPaintSection(FSequencerSectionPainter& InPainter) const
+{
+	const UAnimNotifyState* AnimNotifyState = Section.GetAnimNotifyState();
+	if (AnimNotifyState)
+	{
+		//@TODO: Use AnimNotifyState->NotifyColor
+		return InPainter.PaintSectionBackground();
+	}
+
+	return InPainter.PaintSectionBackground(FLinearColor::Red);
+}
+
+FText FContextualAnimNotifySection::GetSectionTitle() const
+{
+	const UAnimNotifyState* AnimNotifyState = Section.GetAnimNotifyState();
+	if (AnimNotifyState)
+	{
+		return FText::FromString(AnimNotifyState->GetNotifyName());
+	}
+
+	return LOCTEXT("InvalidAnimNotifyState", "Invalid AnimNotifyState");
+}
+
+#undef LOCTEXT_NAMESPACE
\ No newline at end of file
diff --git a/Engine/Plugins/Experimental/Animation/ContextualAnimation/Source/ContextualAnimationEditor/Private/ContextualAnimMovieSceneSequence.cpp b/Engine/Plugins/Experimental/Animation/ContextualAnimation/Source/ContextualAnimationEditor/Private/ContextualAnimMovieSceneSequence.cpp
new file mode 100644
--- /dev/null
+++ b/Engine/Plugins/Experimental/Animation/ContextualAnimation/Source/ContextualAnimationEditor/Private/ContextualAnimMovieSceneSequence.cpp
@@ -0,0 +1,82 @@
+// Copyright Epic Games, Inc. All Rights Reserved.
+
+#include "ContextualAnimMovieSceneSequence.h"
+#include "Tracks/MovieSceneEventTrack.h"
+#include "Tracks/MovieSceneSkeletalAnimationTrack.h"
+#include "Tracks/MovieSceneActorReferenceTrack.h"
+#include "ContextualAnimMovieSceneNotifyTrack.h"
+#include "ContextualAnimViewModel.h"
+
+UContextualAnimMovieSceneSequence::UContextualAnimMovieSceneSequence(const FObjectInitializer& ObjectInitializer)
+	: Super(ObjectInitializer)
+{
+
+}
+
+ETrackSupport UContextualAnimMovieSceneSequence::IsTrackSupported(TSubclassOf<class UMovieSceneTrack> InTrackClass) const
+{ 
+	//@TODO: Remove UMovieSceneSkeletalAnimationTrack and UMovieSceneEventTrack from here after adding a custom track to represent the animation
+	if (InTrackClass == UMovieSceneSkeletalAnimationTrack::StaticClass() ||
+		InTrackClass == UMovieSceneEventTrack::StaticClass() ||
+		InTrackClass == UMovieSceneActorReferenceTrack::StaticClass() ||
+		InTrackClass == UContextualAnimMovieSceneNotifyTrack::StaticClass())
+	{
+		return ETrackSupport::Supported;
+	}
+	
+	return Super::IsTrackSupported(InTrackClass);
+}
+
+void UContextualAnimMovieSceneSequence::Initialize(const TSharedRef<FContextualAnimViewModel>& ViewModelRef)
+{
+	ViewModelPtr = ViewModelRef;
+}
+
+void UContextualAnimMovieSceneSequence::BindPossessableObject(const FGuid& ObjectId, UObject& PossessedObject, UObject* Context)
+{
+	if (AActor* Actor = Cast<AActor>(&PossessedObject))
+	{
+		BoundActors.Add(ObjectId, Actor);
+	}
+}
+
+bool UContextualAnimMovieSceneSequence::CanPossessObject(UObject& Object, UObject* InPlaybackContext) const
+{
+	return true;
+}
+
+void UContextualAnimMovieSceneSequence::LocateBoundObjects(const FGuid& ObjectId, UObject* Context, TArray<UObject*, TInlineAllocator<1>>& OutObjects) const
+{
+	const TWeakObjectPtr<AActor>* WeakActorPtr = BoundActors.Find(ObjectId);
+	if (WeakActorPtr && WeakActorPtr->IsValid())
+	{
+		OutObjects.Add(WeakActorPtr->Get());
+	}
+}
+
+UObject* UContextualAnimMovieSceneSequence::GetParentObject(UObject* Object) const
+{
+	return nullptr;
+}
+
+void UContextualAnimMovieSceneSequence::UnbindPossessableObjects(const FGuid& ObjectId)
+{
+	BoundActors.Remove(ObjectId);
+}
+
+void UContextualAnimMovieSceneSequence::UnbindObjects(const FGuid& ObjectId, const TArray<UObject*>& InObjects, UObject* Context)
+{
+	BoundActors.Remove(ObjectId);
+}
+
+void UContextualAnimMovieSceneSequence::UnbindInvalidObjects(const FGuid& ObjectId, UObject* Context)
+{
+	BoundActors.Remove(ObjectId);
+}
+
+UMovieScene* UContextualAnimMovieSceneSequence::GetMovieScene() const
+{
+	UMovieScene* MovieScene = GetViewModel().GetMovieScene();
+	checkf(MovieScene, TEXT("ContextualAnim sequence not initialized"));
+	return MovieScene;
+}
\ No newline at end of file
diff --git a/Engine/Plugins/Experimental/Animation/ContextualAnimation/Source/ContextualAnimationEditor/Private/ContextualAnimPreviewManager.cpp b/Engine/Plugins/Experimental/Animation/ContextualAnimation/Source/ContextualAnimationEditor/Private/ContextualAnimPreviewManager.cpp
--- a/Engine/Plugins/Experimental/Animation/ContextualAnimation/Source/ContextualAnimationEditor/Private/ContextualAnimPreviewManager.cpp
+++ b/Engine/Plugins/Experimental/Animation/ContextualAnimation/Source/ContextualAnimationEditor/Private/ContextualAnimPreviewManager.cpp
@@ -1,13 +1,41 @@
 // Copyright Epic Games, Inc. All Rights Reserved.
 
 #include "ContextualAnimPreviewManager.h"
+#include "ContextualAnimTypes.h"
 #include "ContextualAnimSceneAsset.h"
-#include "ContextualAnimEdMode.h"
-
+#include "ContextualAnimUtilities.h"
+#include "Animation/AnimMontage.h"
+#include "Animation/AnimInstance.h"
 #include "GameFramework/Character.h"
 #include "GameFramework/CharacterMovementComponent.h"
+#include "Camera/CameraComponent.h"
+#include "MotionWarpingComponent.h"
 #include "NavigationSystem.h"
 #include "AIController.h"
+#include "SkeletalDebugRendering.h"
+
+// FContextualAnimPreviewActorData
+//============================================================================================================
+
+void FContextualAnimPreviewActorData::ResetActorTransform(float Time)
+{
+	AActor* PreviewActor = GetActor();
+
+	const USkeletalMeshComponent* SkelMeshComp = UContextualAnimUtilities::TryGetSkeletalMeshComponent(PreviewActor);
+
+	const FTransform RootTransform = UContextualAnimUtilities::ExtractRootTransformFromAnimation(GetAnimation(), Time);
+	const FTransform StartTransform = SkelMeshComp->GetRelativeTransform().Inverse() * RootTransform;
+
+	PreviewActor->SetActorLocationAndRotation(StartTransform.GetLocation(), StartTransform.GetRotation());
+
+	if (UCharacterMovementComponent* MovementComp = PreviewActor->FindComponentByClass<UCharacterMovementComponent>())
+	{
+		MovementComp->StopMovementImmediately();
+	}
+}
+
+// UContextualAnimPreviewManager
+//============================================================================================================
 
 UContextualAnimPreviewManager::UContextualAnimPreviewManager(const FObjectInitializer& ObjectInitializer)
 	: Super(ObjectInitializer)
@@ -15,27 +43,65 @@ UContextualAnimPreviewManager::UContextualAnimPreviewManager(const FObjectInitia
 
 }
 
+void UContextualAnimPreviewManager::Initialize(UWorld& World, const UContextualAnimSceneAsset& SceneAsset)
+{
+	WorldPtr = &World;
+	SceneAssetPtr = &SceneAsset;
+}
+
+UWorld* UContextualAnimPreviewManager::GetWorld() const
+{
+	return WorldPtr.Get();
+}
+
+const UContextualAnimSceneAsset* UContextualAnimPreviewManager::GetSceneAsset() const
+{
+	check(SceneAssetPtr.IsValid());
+	return SceneAssetPtr.Get();
+}
+
+FName UContextualAnimPreviewManager::FindRoleByGuid(const FGuid& Guid) const
+{ 
+	const FContextualAnimPreviewActorData* Result = PreviewActorsData.FindByPredicate([&Guid](const FContextualAnimPreviewActorData& Data) {
+		return Data.Guid == Guid;
+	});
+
+	return Result ? Result->Role : NAME_None;
+}
+
+UAnimMontage* UContextualAnimPreviewManager::FindAnimationByRole(const FName& Role) const
+{
+	//@TODO: Use the first animation for now
+	const FContextualAnimData* DataPtr = GetSceneAsset()->GetAnimDataForRoleAtIndex(Role, 0);
+	return DataPtr ? DataPtr->Animation : nullptr;
+}
+
+UAnimMontage* UContextualAnimPreviewManager::FindAnimationByGuid(const FGuid& Guid) const
+{
+	return FindAnimationByRole(FindRoleByGuid(Guid));
+}
+
 void UContextualAnimPreviewManager::MoveForward(float Value)
 {
-	if (TestCharacter.IsValid())
+	if (ControlledCharacter.IsValid())
 	{
-		const FVector WorldDirection = FRotationMatrix(TestCharacter->GetActorRotation()).GetScaledAxis(EAxis::X);
-		TestCharacter->AddMovementInput(WorldDirection, Value);
+		const FVector WorldDirection = FRotationMatrix(ControlledCharacter->GetActorRotation()).GetScaledAxis(EAxis::X);
+		ControlledCharacter->AddMovementInput(WorldDirection, Value);
 	}
 }
 
 void UContextualAnimPreviewManager::MoveRight(float Value)
 {
-	if(TestCharacter.IsValid())
+	if(ControlledCharacter.IsValid())
 	{
-		const FVector WorldDirection = FRotationMatrix(TestCharacter->GetActorRotation()).GetScaledAxis(EAxis::Y);
-		TestCharacter->AddMovementInput(WorldDirection, Value);
+		const FVector WorldDirection = FRotationMatrix(ControlledCharacter->GetActorRotation()).GetScaledAxis(EAxis::Y);
+		ControlledCharacter->AddMovementInput(WorldDirection, Value);
 	}
 }
 
 void UContextualAnimPreviewManager::MoveToLocation(const FVector& GoalLocation)
 {
-	AAIController* Controller = TestCharacter.IsValid() ? Cast<AAIController>(TestCharacter->GetController()) : nullptr;
+	AAIController* Controller = ControlledCharacter.IsValid() ? Cast<AAIController>(ControlledCharacter->GetController()) : nullptr;
 	if (Controller)
 	{
 		UNavigationSystemV1* NavSys = FNavigationSystem::GetCurrent<UNavigationSystemV1>(Controller->GetWorld());
@@ -46,65 +112,145 @@ void UContextualAnimPreviewManager::MoveToLocation(const FVector& GoalLocation)
 	}
 }
 
-void UContextualAnimPreviewManager::SpawnPreviewActors(const UContextualAnimSceneAsset* SceneAsset, const FTransform& SceneOrigin)
+AActor* UContextualAnimPreviewManager::SpawnPreviewActor(const FName& Role, const FContextualAnimData& Data)
 {
-	if (SceneAsset)
+	UClass* PreviewClass = GetSceneAsset()->GetTrackSettings(Role)->PreviewActorClass;
+	const FTransform SpawnTransform = (Data.AlignmentData.ExtractTransformAtTime(0, 0.f));
+
+	FActorSpawnParameters Params;
+	Params.SpawnCollisionHandlingOverride = ESpawnActorCollisionHandlingMethod::AdjustIfPossibleButAlwaysSpawn;
+	AActor* PreviewActor = GetWorld()->SpawnActor<AActor>(PreviewClass, SpawnTransform, Params);
+
+	if (ACharacter* PreviewCharacter = Cast<ACharacter>(PreviewActor))
 	{
-		PreviewActors.Reset();
+		PreviewCharacter->bUseControllerRotationYaw = false;
 
-		//DrawDebugCoordinateSystem(GetWorld(), ToWorldTransform.GetLocation(), ToWorldTransform.Rotator(), 20.f, false, 10.f, 0, 1.f);
+		if (UCharacterMovementComponent* CharacterMovementComp = PreviewCharacter->GetCharacterMovement())
+		{
+			CharacterMovementComp->bOrientRotationToMovement = true;
+			CharacterMovementComp->bUseControllerDesiredRotation = false;
+			CharacterMovementComp->RotationRate = FRotator(0.f, 540.0, 0.f);
+			CharacterMovementComp->bRunPhysicsWithNoController = true;
 
-		for (const auto& Entry : SceneAsset->DataContainer)
+			CharacterMovementComp->SetMovementMode(Data.bRequireFlyingMode ? EMovementMode::MOVE_Flying : EMovementMode::MOVE_Walking);
+		}
+
+		if (UCameraComponent* CameraComp = PreviewCharacter->FindComponentByClass<UCameraComponent>())
 		{
-			if(Entry.Value.AnimDataContainer.Num() > 0)
-			{
-				const FName& Role = Entry.Key;
-				const FContextualAnimData& Data = Entry.Value.AnimDataContainer[0];
+			CameraComp->DestroyComponent();
+		}
+	}
 
-				const FTransform SpawnTransform = (Data.AlignmentData.ExtractTransformAtTime(0, 0.f) * SceneOrigin);
+	UE_LOG(LogContextualAnim, Log, TEXT("FContextualAnimViewModel::SpawnPreviewActor. Spawned preview Actor: %s at Loc: %s Rot: %s Role: %s"),
+		*GetNameSafe(PreviewActor), *SpawnTransform.GetLocation().ToString(), *SpawnTransform.Rotator().ToString(), *Role.ToString());
 
-				UClass* PreviewClass = SceneAsset->GetTrackSettings(Role)->PreviewActorClass;
-				if (!PreviewClass && DefaultPreviewClass)
-				{
-					PreviewClass = DefaultPreviewClass.Get();
-				}
+	return PreviewActor;
+}
 
-				AActor* PreviewActor = SpawnPreviewActor(PreviewClass, SpawnTransform);
-				if (PreviewActor)
-				{
-					PreviewActors.Add(Role, PreviewActor);
+void UContextualAnimPreviewManager::AddPreviewActor(AActor& Actor, const FName& Role, const FGuid& Guid, UAnimMontage& Animation)
+{
+	FContextualAnimPreviewActorData NewData;
+	NewData.Actor = &Actor;
+	NewData.Role = Role;
+	NewData.Guid = Guid;
+	NewData.Animation = &Animation;
+
+	PreviewActorsData.Add(NewData);
+
+	if (UAnimInstance* AnimInstance = UContextualAnimUtilities::TryGetAnimInstance(&Actor))
+	{
+		AnimInstance->Montage_Play(&Animation);
+		AnimInstance->Montage_Pause(&Animation);
+	}
+}
 
-					if (!TestCharacter.IsValid())
-					{
-						TestCharacter = Cast<ACharacter>(PreviewActor);
-					}
+void UContextualAnimPreviewManager::DisableCollisionBetweenActors()
+{
+	for (const FContextualAnimPreviewActorData& PreviewActorData : PreviewActorsData)
+	{
+		for (const FContextualAnimPreviewActorData& OtherPreviewActorData : PreviewActorsData)
+		{
+			if (PreviewActorData.Actor != OtherPreviewActorData.Actor)
+			{
+				if (ACharacter* PreviewCharacter = Cast<ACharacter>(PreviewActorData.Actor.Get()))
+				{
+					PreviewCharacter->MoveIgnoreActorAdd(OtherPreviewActorData.Actor.Get());
 				}
 			}
 		}
 	}
 }
 
-AActor* UContextualAnimPreviewManager::SpawnPreviewActor(UClass* Class, const FTransform& SpawnTransform) const
+void UContextualAnimPreviewManager::PreviewTimeChanged(EMovieScenePlayerStatus::Type PreviousStatus, float PreviousTime, EMovieScenePlayerStatus::Type CurrentStatus, float CurrentTime, float PlaybackSpeed)
 {
-	AActor* PreviewActor = FContextualAnimEdMode::Get().GetWorld()->SpawnActor<AActor>(Class, SpawnTransform);
+	//UE_LOG(LogContextualAnim, Log, TEXT("PreviewTimeChanged PrevStatus: %d CurrStatus: %d PrevTime: %f CurrTime: %f Speed: %f"), (int32)PreviousStatus, (int32)CurrentStatus, PreviousTime, CurrentTime, PlaybackSpeed);
 
-	if (ACharacter* PreviewCharacter = Cast<ACharacter>(PreviewActor))
+	for (FContextualAnimPreviewActorData& PreviewActorData : PreviewActorsData)
 	{
-		PreviewCharacter->bUseControllerRotationYaw = false;
-		if (UCharacterMovementComponent* CharacterMovementComp = PreviewCharacter->GetCharacterMovement())
+		UAnimMontage* Animation = PreviewActorData.Animation.Get();
+
+		// Ignore static actors
+		if(Animation == nullptr)
 		{
-			CharacterMovementComp->bOrientRotationToMovement = true;
-			CharacterMovementComp->bUseControllerDesiredRotation = false;
-			CharacterMovementComp->RotationRate = FRotator(0.f, 540.0, 0.f);
+			continue;
+		}
+
+		FAnimMontageInstance* MontageInstance = nullptr;
+		if (UAnimInstance* AnimInstance = UContextualAnimUtilities::TryGetAnimInstance(PreviewActorData.GetActor()))
+		{
+			MontageInstance = AnimInstance->GetActiveMontageInstance();
+
+			// Ensure animation is always playing
+			if (MontageInstance == nullptr)
+			{
+				AnimInstance->Montage_Play(Animation, 1.f, EMontagePlayReturnType::Duration, CurrentTime);
+				AnimInstance->Montage_Pause(Animation);
+				
+				MontageInstance = AnimInstance->GetActiveMontageInstance();
+				check(MontageInstance);
+			}
 		}
 
-		if (!PreviewCharacter->AIControllerClass || !PreviewCharacter->AIControllerClass->IsChildOf<AAIController>())
+		const float AnimPlayLength = Animation->GetPlayLength();
+		PreviousTime = FMath::Clamp(PreviousTime, 0.f, AnimPlayLength);
+		CurrentTime = FMath::Clamp(CurrentTime, 0.f, AnimPlayLength);
+
+		if (CurrentStatus == EMovieScenePlayerStatus::Stopped || CurrentStatus == EMovieScenePlayerStatus::Scrubbing)
 		{
-			PreviewCharacter->AIControllerClass = AAIController::StaticClass();
+			PreviewActorData.ResetActorTransform(CurrentTime);
+
+			if (MontageInstance->IsPlaying())
+			{
+				MontageInstance->Pause();
+			}
+
+			MontageInstance->SetPosition(CurrentTime);
 		}
+		else if (CurrentStatus == EMovieScenePlayerStatus::Playing)
+		{
+			if(PlaybackSpeed > 0.f && CurrentTime < PreviousTime)
+			{
+				PreviewActorData.ResetActorTransform(CurrentTime);
+			}
 
-		PreviewCharacter->SpawnDefaultController();
+			if (!MontageInstance->IsPlaying())
+			{
+				MontageInstance->SetPlaying(true);
+			}
+		}
 	}
+}
 
-	return PreviewActor;
+void UContextualAnimPreviewManager::Reset()
+{
+	for (const FContextualAnimPreviewActorData& PreviewActorData : PreviewActorsData)
+	{
+		PreviewActorData.Actor->Destroy();
+	}
+
+	PreviewActorsData.Reset();
+}
+
+void UContextualAnimPreviewManager::Draw(const FSceneView* View, FPrimitiveDrawInterface* PDI)
+{
 }
\ No newline at end of file
diff --git a/Engine/Plugins/Experimental/Animation/ContextualAnimation/Source/ContextualAnimationEditor/Private/ContextualAnimPreviewScene.cpp b/Engine/Plugins/Experimental/Animation/ContextualAnimation/Source/ContextualAnimationEditor/Private/ContextualAnimPreviewScene.cpp
new file mode 100644
--- /dev/null
+++ b/Engine/Plugins/Experimental/Animation/ContextualAnimation/Source/ContextualAnimationEditor/Private/ContextualAnimPreviewScene.cpp
@@ -0,0 +1,48 @@
+// Copyright Epic Games, Inc. All Rights Reserved.
+
+#include "ContextualAnimPreviewScene.h"
+#include "ContextualAnimAssetEditorToolkit.h"
+#include "GameFramework/WorldSettings.h"
+#include "EngineUtils.h"
+
+FContextualAnimPreviewScene::FContextualAnimPreviewScene(ConstructionValues CVS, const TSharedRef<FContextualAnimAssetEditorToolkit>& EditorToolkit)
+	: FAdvancedPreviewScene(CVS)
+	, EditorToolkitPtr(EditorToolkit)
+{
+	// Disable killing actors outside of the world
+	AWorldSettings* WorldSettings = GetWorld()->GetWorldSettings(true);
+	WorldSettings->bEnableWorldBoundsChecks = false;
+
+	// Spawn an owner for FloorMeshComponent so CharacterMovementComponent can detect it as a valid floor and slide along it
+	{
+		AActor* FloorActor = GetWorld()->SpawnActor<AActor>(AActor::StaticClass(), FTransform());
+		check(FloorActor);
+
+		static const FString NewName = FString(TEXT("FloorComponent"));
+		FloorMeshComponent->Rename(*NewName, FloorActor);
+
+		FloorActor->SetRootComponent(FloorMeshComponent);
+	}
+}
+
+void FContextualAnimPreviewScene::Tick(float InDeltaTime)
+{
+	FAdvancedPreviewScene::Tick(InDeltaTime);
+
+	// Trigger Begin Play in this preview world.
+	// This is needed for the CharacterMovementComponent to be able to switch to falling mode. See: UCharacterMovementComponent::StartFalling
+	if (PreviewWorld && !PreviewWorld->bBegunPlay)
+	{
+		for (FActorIterator It(PreviewWorld); It; ++It)
+		{
+			It->DispatchBeginPlay();
+		}
+
+		PreviewWorld->bBegunPlay = true;
+	}
+
+	if (!GIntraFrameDebuggingGameThread)
+	{
+		GetWorld()->Tick(LEVELTICK_All, InDeltaTime);
+	}
+}
\ No newline at end of file
diff --git a/Engine/Plugins/Experimental/Animation/ContextualAnimation/Source/ContextualAnimationEditor/Private/ContextualAnimTypeActions.cpp b/Engine/Plugins/Experimental/Animation/ContextualAnimation/Source/ContextualAnimationEditor/Private/ContextualAnimTypeActions.cpp
new file mode 100644
--- /dev/null
+++ b/Engine/Plugins/Experimental/Animation/ContextualAnimation/Source/ContextualAnimationEditor/Private/ContextualAnimTypeActions.cpp
@@ -0,0 +1,43 @@
+// Copyright Epic Games, Inc. All Rights Reserved.
+
+#include "ContextualAnimTypeActions.h"
+#include "ContextualAnimAssetEditorToolkit.h"
+#include "ContextualAnimSceneAsset.h"
+
+#define LOCTEXT_NAMESPACE "ContextualAnimTypeActions"
+
+FText FContextualAnimTypeActions::GetName() const
+{
+	return LOCTEXT("ContextualAnimTypeActionsName", "Contextual Anim Scene");
+}
+
+FColor FContextualAnimTypeActions::GetTypeColor() const
+{
+	return FColor(129, 196, 115);
+}
+
+UClass* FContextualAnimTypeActions::GetSupportedClass() const
+{
+	return UContextualAnimSceneAsset::StaticClass();
+}
+
+void FContextualAnimTypeActions::OpenAssetEditor(const TArray<UObject*>& InObjects, TSharedPtr<class IToolkitHost> EditWithinLevelEditor)
+{
+	const EToolkitMode::Type Mode = EditWithinLevelEditor.IsValid() ? EToolkitMode::WorldCentric : EToolkitMode::Standalone;
+
+	for (auto ObjIt = InObjects.CreateConstIterator(); ObjIt; ++ObjIt)
+	{
+		if (UContextualAnimSceneAsset* Asset = Cast<UContextualAnimSceneAsset>(*ObjIt))
+		{
+			TSharedRef<FContextualAnimAssetEditorToolkit> NewEditor(new FContextualAnimAssetEditorToolkit());
+			NewEditor->InitAssetEditor(Mode, EditWithinLevelEditor, Asset);
+		}
+	}
+}
+
+uint32 FContextualAnimTypeActions::GetCategories()
+{
+	return EAssetTypeCategories::Animation;
+}
+
+#undef LOCTEXT_NAMESPACE
\ No newline at end of file
diff --git a/Engine/Plugins/Experimental/Animation/ContextualAnimation/Source/ContextualAnimationEditor/Private/ContextualAnimViewModel.cpp b/Engine/Plugins/Experimental/Animation/ContextualAnimation/Source/ContextualAnimationEditor/Private/ContextualAnimViewModel.cpp
new file mode 100644
--- /dev/null
+++ b/Engine/Plugins/Experimental/Animation/ContextualAnimation/Source/ContextualAnimationEditor/Private/ContextualAnimViewModel.cpp
@@ -0,0 +1,322 @@
+// Copyright Epic Games, Inc. All Rights Reserved.
+
+#include "ContextualAnimViewModel.h"
+#include "Modules/ModuleManager.h"
+#include "MovieScene.h"
+#include "ISequencerModule.h"
+#include "Tracks/MovieSceneSkeletalAnimationTrack.h"
+#include "Tracks/MovieSceneEventTrack.h"
+#include "Sections/MovieSceneEventRepeaterSection.h"
+#include "ContextualAnimMovieSceneSequence.h"
+#include "ContextualAnimMovieSceneNotifyTrackEditor.h"
+#include "ContextualAnimMovieSceneNotifyTrack.h"
+#include "ContextualAnimMovieSceneNotifySection.h"
+#include "ContextualAnimPreviewScene.h"
+#include "ContextualAnimPreviewManager.h"
+#include "Animation/AnimMontage.h"
+#include "Animation/AnimInstance.h"
+#include "Animation/AnimNotifies/AnimNotifyState.h"
+#include "EngineUtils.h"
+#include "ContextualAnimTypes.h"
+#include "ContextualAnimUtilities.h"
+#include "ContextualAnimSceneAsset.h"
+#include "AnimNotifyState_IKWindow.h"
+
+FContextualAnimViewModel::FContextualAnimViewModel()
+	: SceneAsset(nullptr)
+	, MovieSceneSequence(nullptr)
+	, MovieScene(nullptr)
+{
+}
+
+FContextualAnimViewModel::~FContextualAnimViewModel()
+{
+	if (Sequencer.IsValid())
+	{
+		Sequencer->OnMovieSceneDataChanged().RemoveAll(this);
+		Sequencer->OnGlobalTimeChanged().RemoveAll(this);
+		Sequencer.Reset();
+	}
+}
+
+void FContextualAnimViewModel::AddReferencedObjects(FReferenceCollector& Collector)
+{
+	Collector.AddReferencedObject(SceneAsset);
+	Collector.AddReferencedObject(PreviewManager);
+	Collector.AddReferencedObject(MovieSceneSequence);
+	Collector.AddReferencedObject(MovieScene);
+}
+
+TSharedPtr<ISequencer> FContextualAnimViewModel::GetSequencer()
+{
+	return Sequencer;
+}
+
+void FContextualAnimViewModel::Initialize(UContextualAnimSceneAsset* InSceneAsset, const TSharedRef<FContextualAnimPreviewScene>& InPreviewScene)
+{
+	SceneAsset = InSceneAsset;
+	PreviewScenePtr = InPreviewScene;
+
+	PreviewManager = NewObject<UContextualAnimPreviewManager>(GetTransientPackage());
+	PreviewManager->Initialize(*PreviewScenePtr.Pin()->GetWorld(), *SceneAsset);
+
+	CreateSequencer();
+	
+	RefreshSequencerTracks();
+}
+
+UAnimMontage* FContextualAnimViewModel::FindAnimationByGuid(const FGuid& Guid) const
+{
+	return PreviewManager->FindAnimationByGuid(Guid);
+}
+
+void FContextualAnimViewModel::CreateSequencer()
+{
+	MovieSceneSequence = NewObject<UContextualAnimMovieSceneSequence>(GetTransientPackage());
+	MovieSceneSequence->Initialize(AsShared());
+
+	MovieScene = NewObject<UMovieScene>(MovieSceneSequence, FName("ContextualAnimMovieScene"), RF_Transactional);
+	MovieScene->SetDisplayRate(FFrameRate(30, 1));
+
+	FSequencerViewParams ViewParams(TEXT("ContextualAnimSequenceSettings"));
+	{
+		ViewParams.UniqueName = "ContextualAnimSequenceEditor";
+		//ViewParams.OnGetAddMenuContent = OnGetSequencerAddMenuContent;
+		//ViewParams.OnGetPlaybackSpeeds = ISequencer::FOnGetPlaybackSpeeds::CreateRaw(this, &FContextualAnimViewModel::GetPlaybackSpeeds);
+	}
+
+	FSequencerInitParams SequencerInitParams;
+	{
+		SequencerInitParams.ViewParams = ViewParams;
+		SequencerInitParams.RootSequence = MovieSceneSequence;
+		SequencerInitParams.bEditWithinLevelEditor = false;
+		SequencerInitParams.ToolkitHost = nullptr;
+		SequencerInitParams.PlaybackContext.Bind(this, &FContextualAnimViewModel::GetPlaybackContext);
+	}
+
+	ISequencerModule& SequencerModule = FModuleManager::LoadModuleChecked< ISequencerModule >("Sequencer");
+	Sequencer = SequencerModule.CreateSequencer(SequencerInitParams);
+	Sequencer->OnMovieSceneDataChanged().AddRaw(this, &FContextualAnimViewModel::SequencerDataChanged);
+	Sequencer->OnGlobalTimeChanged().AddRaw(this, &FContextualAnimViewModel::SequencerTimeChanged);
+	//Sequencer->GetSelectionChangedTracks().AddRaw(this, &FContextualAnimAssetEditorToolkit::SequencerTrackSelectionChanged);
+	//Sequencer->GetSelectionChangedSections().AddRaw(this, &FContextualAnimAssetEditorToolkit::SequencerSectionSelectionChanged);
+	Sequencer->SetPlaybackStatus(EMovieScenePlayerStatus::Stopped);
+}
+
+void FContextualAnimViewModel::RefreshSequencerTracks()
+{
+	// Remove movie scene tracks (if any)
+	for (const FContextualAnimPreviewActorData& Data : PreviewManager->PreviewActorsData)
+	{
+		Data.GetAnimation()->UnregisterOnNotifyChanged(this);
+
+		MovieSceneSequence->GetMovieScene()->RemovePossessable(Data.Guid);
+	}
+
+	// Destroy preview actors (if any)
+	PreviewManager->Reset();
+
+	Sequencer->NotifyMovieSceneDataChanged(EMovieSceneDataChangeType::MovieSceneStructureItemsChanged);
+
+	for (const auto& Entry : SceneAsset->DataContainer)
+	{
+		if (Entry.Value.AnimDataContainer.Num() > 0)
+		{
+			const FName& Role = Entry.Key;
+			const FContextualAnimData& Data = Entry.Value.AnimDataContainer[0]; //@TODO: Use the first entry for now
+
+			UAnimMontage* AnimMontage = Data.Animation;
+			if(AnimMontage == nullptr)
+			{
+				continue;
+			}
+
+			// Spawn preview actor
+			AActor* PreviewActor = PreviewManager->SpawnPreviewActor(Role, Data);
+			check(PreviewActor);
+
+			// Set actor label so the track shows the name of the role it represents
+			PreviewActor->SetActorLabel(Role.ToString(), false);
+
+			// Add preview actors to sequencer
+			const bool bSelectActors = false;
+			TArray<TWeakObjectPtr<AActor>> Actors = { PreviewActor };
+			TArray<FGuid> Guids = Sequencer->AddActors(Actors, bSelectActors);
+			check(Guids.Num() > 0);
+
+			const FGuid& Guid = Guids[0];
+
+			// Add Animation Track
+			{
+				// @TODO: Temporally using an EventTrack to represent the animation since this is just a visual representation of the data. This assumes there is a single section in the montage
+				UMovieSceneEventTrack* AnimTrack = MovieSceneSequence->GetMovieScene()->AddTrack<UMovieSceneEventTrack>(Guid);
+				check(AnimTrack);
+
+				AnimTrack->SetDisplayName(FText::FromString(GetNameSafe(AnimMontage)));
+
+				UMovieSceneSection* NewSection = NewObject<UMovieSceneSection>(AnimTrack, UMovieSceneEventRepeaterSection::StaticClass(), NAME_None, RF_Transactional);
+				check(NewSection);
+
+				FFrameRate TickResolution = MovieSceneSequence->GetMovieScene()->GetTickResolution();
+				FFrameNumber StartFrame(0);
+				FFrameNumber EndFrame = (AnimMontage->GetPlayLength() * TickResolution).RoundToFrame();
+				NewSection->SetRange(TRange<FFrameNumber>::Exclusive(StartFrame, EndFrame));
+
+				AnimTrack->AddSection(*NewSection);
+			}
+
+			// Add Notify Tracks
+			{
+				for (const FAnimNotifyTrack& NotifyTrack : AnimMontage->AnimNotifyTracks)
+				{
+					UContextualAnimMovieSceneNotifyTrack* Track = MovieSceneSequence->GetMovieScene()->AddTrack<UContextualAnimMovieSceneNotifyTrack>(Guid);
+					check(Track);
+
+					Track->Initialize(*AnimMontage, NotifyTrack);
+				}
+
+				// Listen for when the notifies in the animation changes, so we can refresh the notify sections here
+				AnimMontage->RegisterOnNotifyChanged(UAnimSequenceBase::FOnNotifyChanged::CreateSP(this, &FContextualAnimViewModel::OnAnimNotifyChanged, AnimMontage));
+			}
+
+			// Save preview actor data
+			PreviewManager->AddPreviewActor(*PreviewActor, Role, Guid, *AnimMontage);
+		}
+	}
+
+	//@TODO: DisableCollisionBetweenActors should also be called if bDisableCollisionBetweenActors changes while the editor is open
+	if(SceneAsset->bDisableCollisionBetweenActors)
+	{
+		PreviewManager->DisableCollisionBetweenActors();
+	}
+
+	Sequencer->NotifyMovieSceneDataChanged(EMovieSceneDataChangeType::MovieSceneStructureItemsChanged);
+}
+
+void FContextualAnimViewModel::AddActorTrack(const FNewRoleWidgetParams& Params)
+{	
+	UE_LOG(LogContextualAnim, Log, TEXT("FContextualAnimViewModel::AddActorTrack Role: %s PreviewClass: %s Animation: %s"), 
+		*Params.RoleName.ToString(), *GetNameSafe(Params.PreviewClass), *GetNameSafe(Params.Animation));
+
+	FContextualAnimCompositeTrack CompositeTrack;
+	CompositeTrack.Settings.PreviewActorClass = Params.PreviewClass;
+	CompositeTrack.Settings.MeshToComponent = Params.MeshToComponent;
+
+	FContextualAnimData AnimData;
+	AnimData.Animation = Params.Animation;
+	AnimData.bRequireFlyingMode = Params.bRequiresFlyingMode;
+	CompositeTrack.AnimDataContainer.Add(AnimData);
+
+	SceneAsset->DataContainer.Add(Params.RoleName, CompositeTrack);
+
+	SceneAsset->PrecomputeData();
+
+	RefreshSequencerTracks();
+}
+
+UObject* FContextualAnimViewModel::GetPlaybackContext() const
+{
+	check(PreviewScenePtr.IsValid());
+	return PreviewScenePtr.Pin()->GetWorld();
+}
+
+void FContextualAnimViewModel::SequencerTimeChanged()
+{
+	EMovieScenePlayerStatus::Type CurrentStatus = Sequencer->GetPlaybackStatus();
+	const float CurrentSequencerTime = Sequencer->GetGlobalTime().AsSeconds();
+	const float PlaybackSpeed = Sequencer->GetPlaybackSpeed();
+
+	PreviewManager->PreviewTimeChanged(PreviousSequencerStatus, PreviousSequencerTime, CurrentStatus, CurrentSequencerTime, PlaybackSpeed);
+
+	PreviousSequencerStatus = CurrentStatus;
+	PreviousSequencerTime = CurrentSequencerTime;
+}
+
+void FContextualAnimViewModel::SequencerDataChanged(EMovieSceneDataChangeType DataChangeType)
+{
+	UE_LOG(LogContextualAnim, Log, TEXT("FContextualAnimViewModel::OnMovieSceneDataChanged DataChangeType: %d"), (int32)DataChangeType);
+
+	if(DataChangeType == EMovieSceneDataChangeType::TrackValueChanged)
+	{
+		// Update IK AnimNotify's bEnable flag based on the Active state of the section
+		// @TODO: Temp brute-force approach until having a way to override FMovieSceneSection::SetIsActive or something similar
+
+		for (const FContextualAnimPreviewActorData& Data : PreviewManager->PreviewActorsData)
+		{
+			TArray<UMovieSceneTrack*> Tracks = MovieSceneSequence->GetMovieScene()->FindTracks(UContextualAnimMovieSceneNotifyTrack::StaticClass(), Data.Guid);
+			for(UMovieSceneTrack* Track : Tracks)
+			{
+				const TArray<UMovieSceneSection*>& Sections = Track->GetAllSections();
+				for(UMovieSceneSection* Section : Sections)
+				{
+					if(UContextualAnimMovieSceneNotifySection* NotifySection = Cast<UContextualAnimMovieSceneNotifySection>(Section))
+					{
+						if(UAnimNotifyState_IKWindow* IKNotify = Cast<UAnimNotifyState_IKWindow>(NotifySection->GetAnimNotifyState()))
+						{
+							if(IKNotify->bEnable != NotifySection->IsActive())
+							{
+								IKNotify->bEnable = NotifySection->IsActive();
+								IKNotify->MarkPackageDirty();
+							}
+						}
+					}
+				}
+			}
+		}
+	}
+	else if(DataChangeType == EMovieSceneDataChangeType::MovieSceneStructureItemRemoved)
+	{
+		// Remove preview actor and role from the scene asset if an actor track was removed
+		int32 TotalRemoved = 0;
+		for (int32 Idx = PreviewManager->PreviewActorsData.Num() - 1; Idx >= 0; Idx--)
+		{
+			const FContextualAnimPreviewActorData& Data = PreviewManager->PreviewActorsData[Idx];
+
+			// Missing binding means the actor track was removed...
+			FMovieSceneBinding* Binding = MovieSceneSequence->GetMovieScene()->FindBinding(Data.Guid);
+			if (Binding == nullptr)
+			{
+				// Remove role from the scene asset
+				SceneAsset->DataContainer.Remove(Data.Role);
+
+				// Remove preview actor from the scene
+				Data.Actor->Destroy();
+
+				// Remove cached data
+				PreviewManager->PreviewActorsData.RemoveAt(Idx, 1, false);
+
+				TotalRemoved++;
+			}
+		}
+
+		PreviewManager->PreviewActorsData.Shrink();
+
+		if(TotalRemoved > 0)
+		{
+			SceneAsset->PrecomputeData();
+			SceneAsset->MarkPackageDirty();
+		}
+	}
+}
+
+void FContextualAnimViewModel::OnAnimNotifyChanged(UAnimMontage* Animation)
+{
+	// Do not refresh sequencer tracks if the change to the notifies came from us
+	if (bUpdatingAnimationFromSequencer)
+	{
+		return;
+	}
+
+	UE_LOG(LogContextualAnim, Log, TEXT("FContextualAnimViewModel::OnAnimNotifyChanged Anim: %s. Refreshing Sequencer Tracks"), *GetNameSafe(Animation));
+
+	RefreshSequencerTracks();
+}
+
+void FContextualAnimViewModel::AnimationModified(UAnimMontage& Animation)
+{
+	TGuardValue<bool> UpdateGuard(bUpdatingAnimationFromSequencer, true);
+
+	Animation.RefreshCacheData();
+	Animation.PostEditChange();
+	Animation.MarkPackageDirty();
+}
\ No newline at end of file
diff --git a/Engine/Plugins/Experimental/Animation/ContextualAnimation/Source/ContextualAnimationEditor/Private/ContextualAnimViewportClient.cpp b/Engine/Plugins/Experimental/Animation/ContextualAnimation/Source/ContextualAnimationEditor/Private/ContextualAnimViewportClient.cpp
new file mode 100644
--- /dev/null
+++ b/Engine/Plugins/Experimental/Animation/ContextualAnimation/Source/ContextualAnimationEditor/Private/ContextualAnimViewportClient.cpp
@@ -0,0 +1,41 @@
+// Copyright Epic Games, Inc. All Rights Reserved.
+
+#include "ContextualAnimViewportClient.h"
+#include "ContextualAnimAssetEditorToolkit.h"
+#include "ContextualAnimPreviewScene.h"
+#include "ContextualAnimPreviewManager.h"
+#include "SContextualAnimViewport.h"
+#include "AssetEditorModeManager.h"
+#include "UnrealWidget.h"
+#include "ContextualAnimEdMode.h"
+
+FContextualAnimViewportClient::FContextualAnimViewportClient(const TSharedRef<FContextualAnimPreviewScene>& InPreviewScene, const TSharedRef<SContextualAnimViewport>& InViewport, const TSharedRef<FContextualAnimAssetEditorToolkit>& InAssetEditorToolkit)
+	: FEditorViewportClient(nullptr, &InPreviewScene.Get(), StaticCastSharedRef<SEditorViewport>(InViewport))
+	, PreviewScenePtr(InPreviewScene)
+	, AssetEditorToolkitPtr(InAssetEditorToolkit)
+{
+	Widget->SetUsesEditorModeTools(ModeTools.Get());
+	StaticCastSharedPtr<FAssetEditorModeManager>(ModeTools)->SetPreviewScene(&InPreviewScene.Get());
+	ModeTools->SetDefaultMode(FContextualAnimEdMode::EdModeId);
+
+	SetRealtime(true);
+
+	SetWidgetCoordSystemSpace(COORD_Local);
+	ModeTools->SetWidgetMode(UE::Widget::WM_Translate);
+}
+
+void FContextualAnimViewportClient::TrackingStarted(const struct FInputEventState& InInputState, bool bIsDraggingWidget, bool bNudge)
+{
+	ModeTools->StartTracking(this, Viewport);
+}
+
+void FContextualAnimViewportClient::TrackingStopped()
+{
+	ModeTools->EndTracking(this, Viewport);
+	Invalidate();
+}
+
+void FContextualAnimViewportClient::Draw(const FSceneView* View, FPrimitiveDrawInterface* PDI)
+{
+	FEditorViewportClient::Draw(View, PDI);
+}
\ No newline at end of file
diff --git a/Engine/Plugins/Experimental/Animation/ContextualAnimation/Source/ContextualAnimationEditor/Private/ContextualAnimationEditor.cpp b/Engine/Plugins/Experimental/Animation/ContextualAnimation/Source/ContextualAnimationEditor/Private/ContextualAnimationEditor.cpp
--- a/Engine/Plugins/Experimental/Animation/ContextualAnimation/Source/ContextualAnimationEditor/Private/ContextualAnimationEditor.cpp
+++ b/Engine/Plugins/Experimental/Animation/ContextualAnimation/Source/ContextualAnimationEditor/Private/ContextualAnimationEditor.cpp
@@ -3,26 +3,74 @@
 #include "ContextualAnimationEditor.h"
 #include "ContextualAnimEdMode.h"
 #include "ContextualAnimEditorStyle.h"
+#include "ContextualAnimTypeActions.h"
+#include "ContextualAnimMovieSceneNotifyTrackEditor.h"
+#include "ContextualAnimAssetEditorCommands.h"
+#include "DetailCustomizations/ContextualAnimNotifySectionDetailCustom.h"
+#include "ISequencerModule.h"
 
 #define LOCTEXT_NAMESPACE "FContextualAnimationEditorModule"
 
 void FContextualAnimationEditorModule::StartupModule()
 {
+	// Register Asset Editor Commands
+	FContextualAnimAssetEditorCommands::Register();
+
 	FContextualAnimEditorStyle::Initialize();
+
 	FContextualAnimEditorStyle::ReloadTextures();
 
-	FEditorModeRegistry::Get().RegisterMode<FContextualAnimEdMode>(FContextualAnimEdMode::EM_ContextualAnimEdModeId, 
+	// Register Ed Mode used to interact with the preview scene
+	FEditorModeRegistry::Get().RegisterMode<FContextualAnimEdMode>(FContextualAnimEdMode::EdModeId, 
 		LOCTEXT("ContextualAnimEdModeEdModeName", "ContextualAnim"), 
 		FSlateIcon(FContextualAnimEditorStyle::GetStyleSetName(), "ContextualAnimEditor.Icon", "ContextualAnimEditor.Icon"),
-		true,
+		false,
 		9000);
+
+	// Register Contextual Anim Scene Asset Type Actions 
+	IAssetTools& AssetTools = FModuleManager::LoadModuleChecked<FAssetToolsModule>("AssetTools").Get();
+	ContextualAnimAssetActions = MakeShared<FContextualAnimTypeActions>();
+	AssetTools.RegisterAssetTypeActions(ContextualAnimAssetActions.ToSharedRef());
+
+	// Register MovieSceneAnimNotifyTrackEditor
+	ISequencerModule& SequencerModule = FModuleManager::LoadModuleChecked< ISequencerModule >("Sequencer");
+	MovieSceneAnimNotifyTrackEditorHandle = SequencerModule.RegisterTrackEditor(FOnCreateTrackEditor::CreateStatic(&FContextualAnimMovieSceneNotifyTrackEditor::CreateTrackEditor));
+
+	// Register Detail Customizations
+	FPropertyEditorModule& PropertyModule = FModuleManager::LoadModuleChecked<FPropertyEditorModule>("PropertyEditor");;
+	PropertyModule.RegisterCustomClassLayout("ContextualAnimMovieSceneNotifySection", FOnGetDetailCustomizationInstance::CreateStatic(&FContextualAnimNotifySectionDetailCustom::MakeInstance));
 }
 
 void FContextualAnimationEditorModule::ShutdownModule()
 {
-	FContextualAnimEditorStyle::Shutdown();
+	// Unregister Asset Editor Commands
+	FContextualAnimAssetEditorCommands::Unregister();
 
-	FEditorModeRegistry::Get().UnregisterMode(FContextualAnimEdMode::EM_ContextualAnimEdModeId);
+	// Unregister MovieSceneAnimNotifyTrackEditor
+	if (FModuleManager::Get().IsModuleLoaded("Sequencer"))
+	{
+		ISequencerModule& SequencerModule = FModuleManager::Get().GetModuleChecked<ISequencerModule>("Sequencer");
+		SequencerModule.UnRegisterTrackEditor(MovieSceneAnimNotifyTrackEditorHandle);
+	}
+
+	// Unregister Contextual Anim Scene Asset Type Actions 
+	if (FModuleManager::Get().IsModuleLoaded("AssetTools"))
+	{
+		IAssetTools& AssetTools = FModuleManager::GetModuleChecked<FAssetToolsModule>("AssetTools").Get();
+		AssetTools.UnregisterAssetTypeActions(ContextualAnimAssetActions.ToSharedRef());
+	}
+
+	// Unregister Detail Customizations
+	if (FModuleManager::Get().IsModuleLoaded("PropertyEditor"))
+	{
+		FPropertyEditorModule& PropertyModule = FModuleManager::Get().GetModuleChecked<FPropertyEditorModule>("PropertyEditor");
+		PropertyModule.UnregisterCustomClassLayout("ContextualAnimMovieSceneNotifySection");
+	}
+
+	// Unregister Ed Mode
+	FEditorModeRegistry::Get().UnregisterMode(FContextualAnimEdMode::EdModeId);
+
+	FContextualAnimEditorStyle::Shutdown();
 }
 
 FContextualAnimationEditorModule& FContextualAnimationEditorModule::Get()
diff --git a/Engine/Plugins/Experimental/Animation/ContextualAnimation/Source/ContextualAnimationEditor/Private/DetailCustomizations/ContextualAnimNotifySectionDetailCustom.cpp b/Engine/Plugins/Experimental/Animation/ContextualAnimation/Source/ContextualAnimationEditor/Private/DetailCustomizations/ContextualAnimNotifySectionDetailCustom.cpp
new file mode 100644
--- /dev/null
+++ b/Engine/Plugins/Experimental/Animation/ContextualAnimation/Source/ContextualAnimationEditor/Private/DetailCustomizations/ContextualAnimNotifySectionDetailCustom.cpp
@@ -0,0 +1,37 @@
+// Copyright Epic Games, Inc. All Rights Reserved.
+
+#include "DetailCustomizations/ContextualAnimNotifySectionDetailCustom.h"
+#include "DetailLayoutBuilder.h"
+#include "DetailCategoryBuilder.h"
+#include "DetailWidgetRow.h"
+#include "PropertyEditorModule.h"
+#include "Animation/AnimNotifies/AnimNotifyState.h"
+#include "ContextualAnimMovieSceneNotifySection.h"
+
+TSharedRef<IDetailCustomization> FContextualAnimNotifySectionDetailCustom::MakeInstance()
+{
+	return MakeShared<FContextualAnimNotifySectionDetailCustom>();
+}
+
+void FContextualAnimNotifySectionDetailCustom::CustomizeDetails(IDetailLayoutBuilder& DetailBuilder)
+{
+	TArray<TWeakObjectPtr<UObject>> ObjectList;
+	DetailBuilder.GetObjectsBeingCustomized(ObjectList);
+	check(ObjectList.Num() > 0);
+
+	Section = Cast<UContextualAnimMovieSceneNotifySection>(ObjectList[0].Get());
+	check(Section.IsValid());
+
+	// Add new category to show the properties from the actual AnimNotify object
+	if (const FAnimNotifyEvent* NotifyEvent = Section->GetAnimNotifyEvent())
+	{
+		if (UAnimNotifyState* NotifyState = NotifyEvent->NotifyStateClass)
+		{
+			IDetailCategoryBuilder& Category = DetailBuilder.EditCategory(TEXT("Notify"), FText::GetEmpty(), ECategoryPriority::Variable);
+
+			TArray<UObject*> ExternalObjects = { NotifyState };
+			IDetailPropertyRow* PropertyRow = Category.AddExternalObjects(ExternalObjects);
+			PropertyRow->ShouldAutoExpand(true);
+		}
+	}
+}
\ No newline at end of file
diff --git a/Engine/Plugins/Experimental/Animation/ContextualAnimation/Source/ContextualAnimationEditor/Private/SContextualAnimAssetBrowser.cpp b/Engine/Plugins/Experimental/Animation/ContextualAnimation/Source/ContextualAnimationEditor/Private/SContextualAnimAssetBrowser.cpp
new file mode 100644
--- /dev/null
+++ b/Engine/Plugins/Experimental/Animation/ContextualAnimation/Source/ContextualAnimationEditor/Private/SContextualAnimAssetBrowser.cpp
@@ -0,0 +1,64 @@
+// Copyright Epic Games, Inc. All Rights Reserved.
+
+#include "SContextualAnimAssetBrowser.h"
+#include "ContentBrowserModule.h"
+#include "IContentBrowserSingleton.h"
+#include "Animation/AnimSequence.h"
+#include "Animation/AnimMontage.h"
+
+void SContextualAnimAssetBrowser::Construct(const FArguments& InArgs)
+{
+	ChildSlot
+		[
+			SNew(SVerticalBox)
+			+ SVerticalBox::Slot()
+			[
+				SAssignNew(AssetBrowserBox, SBox)
+			]
+		];
+
+	FAssetPickerConfig AssetPickerConfig;
+	AssetPickerConfig.Filter.ClassNames.Add(UAnimMontage::StaticClass()->GetFName());
+	AssetPickerConfig.OnAssetDoubleClicked = FOnAssetSelected::CreateSP(this, &SContextualAnimAssetBrowser::OnAssetDoubleClicked);
+	//AssetPickerConfig.GetCurrentSelectionDelegates.Add(&GetCurrentSelectionDelegate);
+	AssetPickerConfig.bAllowNullSelection = false;
+	AssetPickerConfig.InitialAssetViewType = EAssetViewType::Column;
+	AssetPickerConfig.OnShouldFilterAsset = FOnShouldFilterAsset::CreateSP(this, &SContextualAnimAssetBrowser::OnShouldFilterAsset);
+	AssetPickerConfig.bShowPathInColumnView = true;
+	AssetPickerConfig.bShowTypeInColumnView = false;
+
+	// hide all asset registry columns by default (we only really want the name and path)
+	TArray<UObject::FAssetRegistryTag> AssetRegistryTags;
+	UAnimSequence::StaticClass()->GetDefaultObject()->GetAssetRegistryTags(AssetRegistryTags);
+	for (UObject::FAssetRegistryTag& AssetRegistryTag : AssetRegistryTags)
+	{
+		AssetPickerConfig.HiddenColumnNames.Add(AssetRegistryTag.Name.ToString());
+	}
+
+	// Also hide the type column by default (but allow users to enable it, so don't use bShowTypeInColumnView)
+	AssetPickerConfig.HiddenColumnNames.Add(TEXT("Class"));
+
+	FContentBrowserModule& ContentBrowserModule = FModuleManager::Get().LoadModuleChecked<FContentBrowserModule>(TEXT("ContentBrowser"));
+	AssetBrowserBox->SetContent(ContentBrowserModule.Get().CreateAssetPicker(AssetPickerConfig));
+}
+
+void SContextualAnimAssetBrowser::OnAssetDoubleClicked(const FAssetData& AssetData)
+{
+	if (!AssetData.GetAsset())
+	{
+		return;
+	}
+
+	// @TODO: Show pop up to add new role to the interaction using this animation
+	//UAnimMontage* Animation = Cast<UAnimMontage>(AssetData.GetAsset());
+}
+
+bool SContextualAnimAssetBrowser::OnShouldFilterAsset(const FAssetData& AssetData)
+{
+	if (!AssetData.GetClass()->IsChildOf(UAnimMontage::StaticClass()))
+	{
+		return true;
+	}
+
+	return false;
+}
\ No newline at end of file
diff --git a/Engine/Plugins/Experimental/Animation/ContextualAnimation/Source/ContextualAnimationEditor/Private/SContextualAnimViewport.cpp b/Engine/Plugins/Experimental/Animation/ContextualAnimation/Source/ContextualAnimationEditor/Private/SContextualAnimViewport.cpp
new file mode 100644
--- /dev/null
+++ b/Engine/Plugins/Experimental/Animation/ContextualAnimation/Source/ContextualAnimationEditor/Private/SContextualAnimViewport.cpp
@@ -0,0 +1,28 @@
+// Copyright Epic Games, Inc. All Rights Reserved.
+
+#include "SContextualAnimViewport.h"
+#include "ContextualAnimViewportClient.h"
+#include "Toolkits/AssetEditorToolkit.h"
+
+void SContextualAnimViewport::Construct(const FArguments& InArgs, const FContextualAnimViewportRequiredArgs& InRequiredArgs)
+{
+	PreviewScenePtr = InRequiredArgs.PreviewScene;
+	AssetEditorToolkitPtr = InRequiredArgs.AssetEditorToolkit;
+
+	SEditorViewport::Construct(
+		SEditorViewport::FArguments()
+		.IsEnabled(FSlateApplication::Get().GetNormalExecutionAttribute())
+		.AddMetaData<FTagMetaData>(TEXT("AnimationTools.Viewport"))
+	);
+}
+
+TSharedRef<FEditorViewportClient> SContextualAnimViewport::MakeEditorViewportClient()
+{
+	ViewportClient = MakeShared<FContextualAnimViewportClient>(PreviewScenePtr.Pin().ToSharedRef(), SharedThis(this), AssetEditorToolkitPtr.Pin().ToSharedRef());
+	ViewportClient->ViewportType = LVT_Perspective;
+	ViewportClient->bSetListenerPosition = false;
+	ViewportClient->SetViewLocation(EditorViewportDefs::DefaultPerspectiveViewLocation);
+	ViewportClient->SetViewRotation(EditorViewportDefs::DefaultPerspectiveViewRotation);
+
+	return ViewportClient.ToSharedRef();
+}
\ No newline at end of file
diff --git a/Engine/Plugins/Experimental/Animation/ContextualAnimation/Source/ContextualAnimationEditor/Public/ContextualAnimAssetEditorCommands.h b/Engine/Plugins/Experimental/Animation/ContextualAnimation/Source/ContextualAnimationEditor/Public/ContextualAnimAssetEditorCommands.h
new file mode 100644
--- /dev/null
+++ b/Engine/Plugins/Experimental/Animation/ContextualAnimation/Source/ContextualAnimationEditor/Public/ContextualAnimAssetEditorCommands.h
@@ -0,0 +1,22 @@
+// Copyright Epic Games, Inc. All Rights Reserved.
+
+#pragma once
+
+#include "CoreMinimal.h"
+#include "Framework/Commands/Commands.h"
+#include "EditorStyleSet.h"
+
+class FContextualAnimAssetEditorCommands : public TCommands<FContextualAnimAssetEditorCommands>
+{
+public:
+	FContextualAnimAssetEditorCommands()
+		: TCommands<FContextualAnimAssetEditorCommands>(TEXT("ContextualAnimAssetEditor"), NSLOCTEXT("Contexts", "ContextualAnim", "Contextual Anim"), NAME_None, FEditorStyle::GetStyleSetName())
+	{
+	}
+
+	virtual void RegisterCommands() override;
+
+public:
+
+	TSharedPtr<FUICommandInfo> ResetPreviewScene;
+};
diff --git a/Engine/Plugins/Experimental/Animation/ContextualAnimation/Source/ContextualAnimationEditor/Public/ContextualAnimAssetEditorToolkit.h b/Engine/Plugins/Experimental/Animation/ContextualAnimation/Source/ContextualAnimationEditor/Public/ContextualAnimAssetEditorToolkit.h
new file mode 100644
--- /dev/null
+++ b/Engine/Plugins/Experimental/Animation/ContextualAnimation/Source/ContextualAnimationEditor/Public/ContextualAnimAssetEditorToolkit.h
@@ -0,0 +1,65 @@
+// Copyright Epic Games, Inc. All Rights Reserved.
+
+#pragma once
+
+#include "CoreMinimal.h"
+#include "Toolkits/AssetEditorToolkit.h"
+#include "Misc/NotifyHook.h"
+
+class IDetailsView;
+class FToolBarBuilder;
+class UContextualAnimSceneAsset;
+class SContextualAnimViewport;
+class SContextualAnimAssetBrowser;
+class FContextualAnimPreviewScene;
+class FContextualAnimViewModel;
+class UContextualAnimPreviewManager;
+
+class FContextualAnimAssetEditorToolkit : public FAssetEditorToolkit, public FNotifyHook
+{
+public:
+
+	FContextualAnimAssetEditorToolkit();
+	virtual ~FContextualAnimAssetEditorToolkit();
+
+	void InitAssetEditor(const EToolkitMode::Type Mode, const TSharedPtr< IToolkitHost >& InitToolkitHost, UContextualAnimSceneAsset* SceneAsset);
+
+	virtual void RegisterTabSpawners(const TSharedRef<FTabManager>& TabManager) override;
+	virtual void UnregisterTabSpawners(const TSharedRef<FTabManager>& TabManager) override;
+	virtual FName GetToolkitFName() const override;
+	virtual FText GetBaseToolkitName() const override;
+	virtual FText GetToolkitName() const override;
+	virtual FLinearColor GetWorldCentricTabColorScale() const override;
+	virtual FString GetWorldCentricTabPrefix() const override;
+
+	UContextualAnimSceneAsset* GetSceneAsset() const;
+	UContextualAnimPreviewManager* GetPreviewManager() const;
+
+	void ResetPreviewScene();
+
+private:
+	
+	TSharedRef<SDockTab> SpawnTab_Viewport(const FSpawnTabArgs& Args);
+	TSharedRef<SDockTab> SpawnTab_AssetDetails(const FSpawnTabArgs& Args);
+	TSharedRef<SDockTab> SpawnTab_Timeline(const FSpawnTabArgs& Args);
+	TSharedRef<SDockTab> SpawnTab_AssetBrowser(const FSpawnTabArgs& Args);
+	TSharedRef<SDockTab> SpawnTab_PreviewSettings(const FSpawnTabArgs& Args);
+
+	void BindCommands();
+	void ExtendToolbar();
+	void FillToolbar(FToolBarBuilder& ToolbarBuilder);
+
+	void OnFinishedChangingProperties(const FPropertyChangedEvent& PropertyChangedEvent);
+
+	TSharedPtr<SContextualAnimViewport> ViewportWidget;
+
+	TSharedPtr<SContextualAnimAssetBrowser> AssetBrowserWidget;
+
+	TSharedPtr<IDetailsView> EditingAssetWidget;
+
+	TSharedPtr<FContextualAnimPreviewScene> PreviewScene;
+
+	TSharedPtr<FContextualAnimViewModel> ViewModel;
+};
+
+
diff --git a/Engine/Plugins/Experimental/Animation/ContextualAnimation/Source/ContextualAnimationEditor/Public/ContextualAnimEdMode.h b/Engine/Plugins/Experimental/Animation/ContextualAnimation/Source/ContextualAnimationEditor/Public/ContextualAnimEdMode.h
--- a/Engine/Plugins/Experimental/Animation/ContextualAnimation/Source/ContextualAnimationEditor/Public/ContextualAnimEdMode.h
+++ b/Engine/Plugins/Experimental/Animation/ContextualAnimation/Source/ContextualAnimationEditor/Public/ContextualAnimEdMode.h
@@ -3,37 +3,28 @@
 #pragma once
 
 #include "CoreMinimal.h"
-#include "Templates/SubclassOf.h"
 #include "EdMode.h"
 
-class FContextualAnimEdModeToolkit;
-
 class FContextualAnimEdMode : public FEdMode
 {
 public:
 
-	const static FEditorModeID EM_ContextualAnimEdModeId;
+	const static FEditorModeID EdModeId;
 
 	FContextualAnimEdMode();
 	virtual ~FContextualAnimEdMode();
 
-	// FEdMode interface
-	virtual void Enter() override;
-	virtual void Exit() override;
-	virtual void Tick(FEditorViewportClient* ViewportClient, float DeltaTime) override;
-	virtual bool InputKey(FEditorViewportClient* ViewportClient, FViewport* Viewport, FKey Key, EInputEvent Event) override;
 	virtual bool HandleClick(FEditorViewportClient* InViewportClient, HHitProxy* HitProxy, const FViewportClick& Click) override;
-	virtual bool UsesToolkits() const override;
-	// End of FEdMode interface
-
-	/** FGCObject interface */
-	virtual void AddReferencedObjects(FReferenceCollector& Collector) override;
-
-	static FContextualAnimEdMode& Get();
+	virtual bool StartTracking(FEditorViewportClient* InViewportClient, FViewport* InViewport) override;
+	virtual bool EndTracking(FEditorViewportClient* InViewportClient, FViewport* InViewport) override;
+	virtual bool InputDelta(FEditorViewportClient* InViewportClient, FViewport* InViewport, FVector& InDrag, FRotator& InRot, FVector& InScale) override;
+	virtual bool AllowWidgetMove() override;
+	virtual bool ShouldDrawWidget() const override;
+	virtual FVector GetWidgetLocation() const override;
 
-	TSharedPtr<FContextualAnimEdModeToolkit> GetContextualAnimEdModeToolkit() const;
+	bool GetHitResultUnderCursor(FHitResult& OutHitResult, FEditorViewportClient* InViewportClient, const FViewportClick& Click) const;
 
-	class UContextualAnimPreviewManager* PreviewManager;
+private:
 
-	bool GetHitResultUnderCursor(FHitResult& OutHitResult, FEditorViewportClient* InViewportClient, const FViewportClick& Click) const;
+	TWeakObjectPtr<class AActor> SelectedActor;
 };
diff --git a/Engine/Plugins/Experimental/Animation/ContextualAnimation/Source/ContextualAnimationEditor/Public/ContextualAnimEdModeSettings.h b/Engine/Plugins/Experimental/Animation/ContextualAnimation/Source/ContextualAnimationEditor/Public/ContextualAnimEdModeSettings.h
deleted file mode 100644
diff --git a/Engine/Plugins/Experimental/Animation/ContextualAnimation/Source/ContextualAnimationEditor/Public/ContextualAnimEdModeToolkit.h b/Engine/Plugins/Experimental/Animation/ContextualAnimation/Source/ContextualAnimationEditor/Public/ContextualAnimEdModeToolkit.h
deleted file mode 100644
diff --git a/Engine/Plugins/Experimental/Animation/ContextualAnimation/Source/ContextualAnimationEditor/Public/ContextualAnimFactory.h b/Engine/Plugins/Experimental/Animation/ContextualAnimation/Source/ContextualAnimationEditor/Public/ContextualAnimFactory.h
new file mode 100644
--- /dev/null
+++ b/Engine/Plugins/Experimental/Animation/ContextualAnimation/Source/ContextualAnimationEditor/Public/ContextualAnimFactory.h
@@ -0,0 +1,20 @@
+// Copyright Epic Games, Inc. All Rights Reserved.
+
+#pragma once
+
+#include "CoreMinimal.h"
+#include "Factories/Factory.h"
+#include "ContextualAnimFactory.generated.h"
+
+UCLASS()
+class UContextualAnimFactory : public UFactory
+{
+	GENERATED_BODY()
+
+public:
+	
+	UContextualAnimFactory(const FObjectInitializer& ObjectInitializer);
+
+	virtual UObject* FactoryCreateNew(UClass* Class, UObject* InParent, FName Name, EObjectFlags Flags, UObject* Context, FFeedbackContext* Warn) override;
+	virtual bool CanCreateNew() const override;
+};
\ No newline at end of file
diff --git a/Engine/Plugins/Experimental/Animation/ContextualAnimation/Source/ContextualAnimationEditor/Public/ContextualAnimMovieSceneNotifySection.h b/Engine/Plugins/Experimental/Animation/ContextualAnimation/Source/ContextualAnimationEditor/Public/ContextualAnimMovieSceneNotifySection.h
new file mode 100644
--- /dev/null
+++ b/Engine/Plugins/Experimental/Animation/ContextualAnimation/Source/ContextualAnimationEditor/Public/ContextualAnimMovieSceneNotifySection.h
@@ -0,0 +1,48 @@
+// Copyright Epic Games, Inc. All Rights Reserved.
+
+#pragma once
+
+#include "CoreMinimal.h"
+#include "MovieSceneSection.h"
+#include "ContextualAnimMovieSceneNotifySection.generated.h"
+
+struct FAnimNotifyEvent;
+class UAnimNotifyState;
+class UContextualAnimMovieSceneNotifyTrack;
+
+/** 
+ * MovieSceneSection used to represent an AnimNotify in Sequencer Panel 
+ * Stores the FGuid of the AnimNotifyEvent this section is presenting 
+ */
+UCLASS()
+class UContextualAnimMovieSceneNotifySection : public UMovieSceneSection
+{
+	GENERATED_BODY()
+
+public:
+
+	/** 
+	 * Initialize this section from the supplied FAnimNotifyEvent. 
+	 * Sets the range of the section to match the range of the anim notify and cache the FGuid 
+	 */
+	void Initialize(const FAnimNotifyEvent& NotifyEvent);
+
+	/** Returns the track owner of this section */
+	const FGuid& GetAnimNotifyEventGuid() const { return AnimNotifyEventGuid; }
+
+	/** Returns the track owner of this section */
+	UContextualAnimMovieSceneNotifyTrack* GetOwnerTrack() const;
+
+	/** Returns the actual AnimNotifyState this section is representing */
+	UAnimNotifyState* GetAnimNotifyState() const;
+
+	/** Returns the actual AnimNotifyEvent this section is representing */
+	FAnimNotifyEvent* GetAnimNotifyEvent() const;
+
+private:
+
+	/** Guid of the actual AnimNotifyEvent this section is representing */
+	UPROPERTY()
+	FGuid AnimNotifyEventGuid;
+
+};
\ No newline at end of file
diff --git a/Engine/Plugins/Experimental/Animation/ContextualAnimation/Source/ContextualAnimationEditor/Public/ContextualAnimMovieSceneNotifyTrack.h b/Engine/Plugins/Experimental/Animation/ContextualAnimation/Source/ContextualAnimationEditor/Public/ContextualAnimMovieSceneNotifyTrack.h
new file mode 100644
--- /dev/null
+++ b/Engine/Plugins/Experimental/Animation/ContextualAnimation/Source/ContextualAnimationEditor/Public/ContextualAnimMovieSceneNotifyTrack.h
@@ -0,0 +1,71 @@
+// Copyright Epic Games, Inc. All Rights Reserved.
+
+#pragma once
+
+#include "CoreMinimal.h"
+#include "MovieSceneNameableTrack.h"
+#include "ContextualAnimMovieSceneNotifyTrack.generated.h"
+
+class UAnimMontage;
+struct FAnimNotifyTrack;
+class UContextualAnimMovieSceneNotifySection;
+class FContextualAnimViewModel;
+
+/** MovieSceneTrack used to show the notifies in the animation the track was created from in Sequencer Panel */
+UCLASS()
+class UContextualAnimMovieSceneNotifyTrack : public UMovieSceneNameableTrack
+{
+	GENERATED_BODY()
+
+public:
+	
+	FContextualAnimViewModel& GetViewModel() const;
+
+	/** 
+	 * Initialize this track from a supplied animation and a notify track 
+	 * Adding all the FAnimNotifyEvents as sections in here and setting
+	 */
+	void Initialize(UAnimMontage& AnimationRef, const FAnimNotifyTrack& NotifyTrack);
+
+	/** Creates a new section for this track */
+	virtual UMovieSceneSection* CreateNewSection() override;
+	
+	/** Adds a new section to this track */
+	virtual void AddSection(UMovieSceneSection& Section) override;
+	
+	/** Returns whether this track supports the supplied section class */
+	virtual bool SupportsType(TSubclassOf<UMovieSceneSection> SectionClass) const override;
+	
+	/** Returns the list of all the sections in this track */
+	virtual const TArray<UMovieSceneSection*>& GetAllSections() const override;
+	
+	/** Returns whether this track contains the supplied section */
+	virtual bool HasSection(const UMovieSceneSection& Section) const override;
+
+	/** Returns whether this track is empty (doesn't have any section in it) */
+	virtual bool IsEmpty() const override;
+	
+	/** Called when the supplied section is removed from Sequencer. Removes the actual Notify the section is representing */
+	virtual void RemoveSection(UMovieSceneSection& Section) override;
+
+	/** Called when a section at the supplied index is removed from Sequencer. Removes the actual Notify the section is representing */
+	virtual void RemoveSectionAt(int32 SectionIndex) override;
+
+#if WITH_EDITOR
+	/** Called when a section in this track is moved. Updates the actual Notify the section is representing */
+	virtual EMovieSceneSectionMovedResult OnSectionMoved(UMovieSceneSection& Section, const FMovieSceneSectionMovedParams& Params) override;
+#endif
+
+	/** Returns a reference to the animation this track was created from */
+	UAnimMontage& GetAnimation() const;
+
+private:
+
+	/** List of sections in this track */
+	UPROPERTY()
+	TArray<TObjectPtr<UMovieSceneSection>> Sections;
+
+	/** Ptr to the animation this track was created from */
+	UPROPERTY()
+	TWeakObjectPtr<UAnimMontage> Animation;
+};
\ No newline at end of file
diff --git a/Engine/Plugins/Experimental/Animation/ContextualAnimation/Source/ContextualAnimationEditor/Public/ContextualAnimMovieSceneNotifyTrackEditor.h b/Engine/Plugins/Experimental/Animation/ContextualAnimation/Source/ContextualAnimationEditor/Public/ContextualAnimMovieSceneNotifyTrackEditor.h
new file mode 100644
--- /dev/null
+++ b/Engine/Plugins/Experimental/Animation/ContextualAnimation/Source/ContextualAnimationEditor/Public/ContextualAnimMovieSceneNotifyTrackEditor.h
@@ -0,0 +1,166 @@
+// Copyright Epic Games, Inc. All Rights Reserved.
+
+#pragma once
+
+#include "CoreMinimal.h"
+#include "MovieSceneTrackEditor.h"
+#include "ContextualAnimTypes.h"
+#include "BoneContainer.h"
+#include "Interfaces/Interface_BoneReferenceSkeletonProvider.h"
+#include "ContextualAnimMovieSceneNotifyTrackEditor.generated.h"
+
+class AActor;
+class UAnimMontage;
+class USkeleton;
+class UContextualAnimSceneAsset;
+class UContextualAnimMovieSceneNotifySection;
+class UContextualAnimMovieSceneNotifyTrack;
+class UContextualAnimMovieSceneSequence;
+
+/** Struct used to construct the New Role Widget */
+USTRUCT()
+struct FNewRoleWidgetParams
+{
+	GENERATED_BODY()
+
+	UPROPERTY(EditAnywhere, Category = "Settings")
+	FName RoleName = NAME_None;
+
+	UPROPERTY(EditAnywhere, Category = "Settings")
+	TSubclassOf<AActor> PreviewClass;
+
+	UPROPERTY(EditAnywhere, Category = "Settings")
+	UAnimMontage* Animation = nullptr;
+
+	UPROPERTY(EditAnywhere, Category = "Settings")
+	bool bRequiresFlyingMode;
+
+	UPROPERTY(EditAnywhere, Category = "Settings")
+	FTransform MeshToComponent = FTransform(FRotator(0.f, -90.f, 0.f));
+
+	bool HasValidData() const;
+};
+
+USTRUCT()
+struct FNewIKTargetParams
+{
+	GENERATED_BODY()
+
+	UPROPERTY(VisibleAnywhere, Category = "Settings")
+	FName SourceRole = NAME_None;
+
+	UPROPERTY(EditAnywhere, Category = "Settings")
+	FBoneReference SourceBone;
+
+	UPROPERTY(EditAnywhere, Category = "Settings")
+	EContextualAnimIKTargetProvider Provider;
+
+	UPROPERTY(EditAnywhere, Category = "Settings", meta = (GetOptions = "GetTargetRoleOptions"))
+	FName TargetRole = NAME_None;
+
+	UPROPERTY(EditAnywhere, Category = "Settings")
+	FBoneReference TargetBone;
+
+	UPROPERTY(EditAnywhere, Category = "Settings")
+	FName GoalName = NAME_None;
+
+};
+
+/** Object used to construct the New IK Target Widget */
+UCLASS()
+class UNewIKTargetWidgetParams : public UObject, public IBoneReferenceSkeletonProvider
+{
+	GENERATED_BODY()
+
+public:
+
+	UPROPERTY(EditAnywhere, Category = "Settings", meta = (ShowOnlyInnerProperties))
+	FNewIKTargetParams Params;
+
+	UNewIKTargetWidgetParams(const FObjectInitializer& ObjectInitializer) : Super(ObjectInitializer) {}
+
+	// ~IBoneReferenceSkeletonProvider Interface
+	virtual USkeleton* GetSkeleton(bool& bInvalidSkeletonIsError, const IPropertyHandle* PropertyHandle = nullptr) override;
+
+	void Reset(const FName& InSourceRole, const UContextualAnimSceneAsset& InSceneAsset);
+
+	bool HasValidData() const;
+
+	const UContextualAnimSceneAsset& GetSceneAsset() const;
+
+	UFUNCTION()
+	TArray<FString> GetTargetRoleOptions() const;
+
+private:
+
+	UPROPERTY()
+	TWeakObjectPtr<const UContextualAnimSceneAsset> SceneAssetPtr;
+
+	UPROPERTY()
+	TArray<FName> CachedRoles;
+};
+
+/** Handles section drawing and manipulation of a MovieSceneNotifyTrack */
+class FContextualAnimMovieSceneNotifyTrackEditor : public FMovieSceneTrackEditor, public FGCObject
+{
+public:
+
+	FContextualAnimMovieSceneNotifyTrackEditor(TSharedRef<ISequencer> InSequencer);
+
+	// ~FGCObject Interface
+	virtual void AddReferencedObjects(FReferenceCollector& Collector) override;
+	virtual FString GetReferencerName() const override { return TEXT("FContextualAnimMovieSceneNotifyTrackEditor"); }
+
+	// ~FMovieSceneTrackEditor Interface
+	virtual TSharedRef<ISequencerSection> MakeSectionInterface(UMovieSceneSection& SectionObject, UMovieSceneTrack& Track, FGuid ObjectBinding) override;
+	virtual bool SupportsType(TSubclassOf<UMovieSceneTrack> TrackClass) const override;
+	virtual bool SupportsSequence(UMovieSceneSequence* InSequence) const override;
+	virtual TSharedPtr<SWidget> BuildOutlinerEditWidget(const FGuid& ObjectBinding, UMovieSceneTrack* Track, const FBuildEditWidgetParams& Params) override;
+	virtual void BuildAddTrackMenu(FMenuBuilder& MenuBuilder) override;
+	virtual void BuildObjectBindingTrackMenu(FMenuBuilder& MenuBuilder, const TArray<FGuid>& ObjectBindings, const UClass* ObjectClass) override;
+
+	static TSharedRef<ISequencerTrackEditor> CreateTrackEditor(TSharedRef<ISequencer> InSequencer);
+
+	UContextualAnimMovieSceneSequence& GetMovieSceneSequence() const;
+
+private:
+
+	TObjectPtr<UNewIKTargetWidgetParams> NewIKTargetWidgetParams;
+
+	TSharedPtr<FStructOnScope> NewRoleWidgetParams;
+
+	void BuildAddTrackSubMenu(FMenuBuilder& MenuBuilder, TArray<FGuid>);
+
+	UContextualAnimMovieSceneNotifySection* CreateNewSection(UContextualAnimMovieSceneNotifyTrack* Track, int32 RowIndex, UClass* NotifyClass);
+
+	UContextualAnimMovieSceneNotifySection* CreateNewIKSection(UContextualAnimMovieSceneNotifyTrack* Track, int32 RowIndex, const FName& GoalName);
+
+	void BuildNewIKTargetSubMenu(FMenuBuilder& MenuBuilder, UContextualAnimMovieSceneNotifyTrack* Track, int32 RowIndex);
+
+	void BuildNewIKTargetWidget(FMenuBuilder& MenuBuilder, UContextualAnimMovieSceneNotifyTrack* Track, int32 RowIndex);
+
+	void AddNewNotifyTrack(TArray<FGuid> ObjectBindings);
+
+	void FillNewNotifyStateMenu(FMenuBuilder& MenuBuilder, bool bIsReplaceWithMenu, UContextualAnimMovieSceneNotifyTrack* Track, int32 RowIndex);
+};
+
+// FContextualAnimNotifySection
+////////////////////////////////////////////////////////////////////////////////////////////////
+
+/** UI portion of a NotifySection in a NotifyTrack */
+class FContextualAnimNotifySection : public ISequencerSection
+{
+public:
+	
+	FContextualAnimNotifySection(UMovieSceneSection& InSection);
+	virtual ~FContextualAnimNotifySection() { }
+
+	virtual UMovieSceneSection* GetSectionObject() override;
+	virtual int32 OnPaintSection(FSequencerSectionPainter& InPainter) const override;
+	virtual FText GetSectionTitle() const override;
+
+private:
+	
+	/** The section we are visualizing */
+	UContextualAnimMovieSceneNotifySection& Section;
+};
\ No newline at end of file
diff --git a/Engine/Plugins/Experimental/Animation/ContextualAnimation/Source/ContextualAnimationEditor/Public/ContextualAnimMovieSceneSequence.h b/Engine/Plugins/Experimental/Animation/ContextualAnimation/Source/ContextualAnimationEditor/Public/ContextualAnimMovieSceneSequence.h
new file mode 100644
--- /dev/null
+++ b/Engine/Plugins/Experimental/Animation/ContextualAnimation/Source/ContextualAnimationEditor/Public/ContextualAnimMovieSceneSequence.h
@@ -0,0 +1,45 @@
+// Copyright Epic Games, Inc. All Rights Reserved.
+
+#pragma once
+
+#include "CoreMinimal.h"
+#include "UObject/ObjectMacros.h"
+#include "MovieSceneSequence.h"
+#include "ContextualAnimMovieSceneSequence.generated.h"
+
+class UMovieScene;
+class AActor;
+class FContextualAnimViewModel;
+
+UCLASS()
+class UContextualAnimMovieSceneSequence : public UMovieSceneSequence
+{
+	GENERATED_BODY()
+
+public:
+
+	UContextualAnimMovieSceneSequence(const FObjectInitializer& ObjectInitializer);
+
+	void Initialize(const TSharedRef<FContextualAnimViewModel>& ViewModelRef);
+
+	virtual void BindPossessableObject(const FGuid& ObjectId, UObject& PossessedObject, UObject* Context) override;
+	virtual bool CanPossessObject(UObject& Object, UObject* InPlaybackContext) const override;
+	virtual void LocateBoundObjects(const FGuid& ObjectId, UObject* Context, TArray<UObject*, TInlineAllocator<1>>& OutObjects) const override;
+	virtual UMovieScene* GetMovieScene() const override;
+	virtual UObject* GetParentObject(UObject* Object) const override;
+	virtual void UnbindPossessableObjects(const FGuid& ObjectId) override;
+	virtual void UnbindObjects(const FGuid& ObjectId, const TArray<UObject*>& InObjects, UObject* Context) override;
+	virtual void UnbindInvalidObjects(const FGuid& ObjectId, UObject* Context) override;
+	virtual ETrackSupport IsTrackSupported(TSubclassOf<class UMovieSceneTrack> InTrackClass) const;
+
+	FContextualAnimViewModel& GetViewModel() const { return *ViewModelPtr.Pin().Get(); }
+
+private:
+
+	//@TODO: Remove this and use the actors from the PreviewSceneManager
+	UPROPERTY()
+	TMap<FGuid, TWeakObjectPtr<AActor>> BoundActors;
+
+	/** Pointer to the view model that owns this sequence. */
+	TWeakPtr<FContextualAnimViewModel> ViewModelPtr;
+};
diff --git a/Engine/Plugins/Experimental/Animation/ContextualAnimation/Source/ContextualAnimationEditor/Public/ContextualAnimPreviewManager.h b/Engine/Plugins/Experimental/Animation/ContextualAnimation/Source/ContextualAnimationEditor/Public/ContextualAnimPreviewManager.h
--- a/Engine/Plugins/Experimental/Animation/ContextualAnimation/Source/ContextualAnimationEditor/Public/ContextualAnimPreviewManager.h
+++ b/Engine/Plugins/Experimental/Animation/ContextualAnimation/Source/ContextualAnimationEditor/Public/ContextualAnimPreviewManager.h
@@ -5,48 +5,94 @@
 #include "CoreMinimal.h"
 #include "UObject/Object.h"
 #include "Templates/SubclassOf.h"
+#include "MovieSceneFwd.h"
 #include "ContextualAnimPreviewManager.generated.h"
 
+class UWorld;
+class FSceneView;
+class FPrimitiveDrawInterface;
 class ACharacter;
+class UAnimMontage;
 class UContextualAnimSceneAsset;
+struct FContextualAnimData;
 
-UCLASS()
-class UContextualAnimPreviewManager : public UObject
+USTRUCT()
+struct FContextualAnimPreviewActorData
 {
 	GENERATED_BODY()
 
-public:
+	/** Preview actor */
+	UPROPERTY()
+	TWeakObjectPtr<AActor> Actor;
 
-	UPROPERTY(EditAnywhere, Category = "Preview")
-	TSubclassOf<ACharacter> DefaultPreviewClass;
+	/** Role this actor is representing */
+	UPROPERTY()
+	FName Role = NAME_None;
 
-	UPROPERTY(EditAnywhere, Category = "Debug")
-	bool bDrawDebugScene;
+	/** MovieScene Object Binding's identifier  */
+	UPROPERTY()
+	FGuid Guid;
 
-	UPROPERTY(EditAnywhere, Category = "Debug", meta = (EditCondition = "bDrawDebugScene"))
-	int32 AnimDataIndex = 0;
+	/** Animation this actor is playing */
+	UPROPERTY()
+	TWeakObjectPtr<UAnimMontage> Animation;
 
-	UPROPERTY(EditAnywhere, Category = "Debug", meta = (EditCondition = "bDrawDebugScene"))
-	float Time;
+	FORCEINLINE AActor* GetActor() const { return Actor.Get(); }
+	FORCEINLINE UAnimMontage* GetAnimation() const { return Animation.Get(); }
 
-	UPROPERTY(EditAnywhere, Category = "Debug", meta = (EditCondition = "bDrawDebugScene"))
-	FTransform ScenePivot;
+	void ResetActorTransform(float Time);
+};
 
-	UPROPERTY()
-	TWeakObjectPtr<ACharacter> TestCharacter;
+UCLASS()
+class UContextualAnimPreviewManager : public UObject
+{
+	GENERATED_BODY()
+
+public:
 
 	UPROPERTY()
-	TMap<FName, AActor*> PreviewActors;
+	TArray<FContextualAnimPreviewActorData> PreviewActorsData;
 
 	UContextualAnimPreviewManager(const FObjectInitializer& ObjectInitializer);
 
-	void SpawnPreviewActors(const UContextualAnimSceneAsset* SceneAsset, const FTransform& SceneOrigin);
+	void Initialize(UWorld& World, const UContextualAnimSceneAsset& SceneAsset);
 
-	AActor* SpawnPreviewActor(UClass* Class, const FTransform& SpawnTransform) const;
+	virtual UWorld* GetWorld() const override;
+	
+	void Draw(const FSceneView* View, FPrimitiveDrawInterface* PDI);
 
-	void MoveForward(float Value);
+	const UContextualAnimSceneAsset* GetSceneAsset() const;
 
-	void MoveRight(float Value);
+	void DisableCollisionBetweenActors();
 
+	void PreviewTimeChanged(EMovieScenePlayerStatus::Type PreviousStatus, float PreviousTime, EMovieScenePlayerStatus::Type CurrentStatus, float CurrentTime, float PlaybackSpeed);
+
+	void MoveForward(float Value);
+	
+	void MoveRight(float Value);
+	
 	void MoveToLocation(const FVector& GoalLocation);
+
+	AActor* SpawnPreviewActor(const FName& Role, const FContextualAnimData& Data);
+
+	void AddPreviewActor(AActor& Actor, const FName& Role, const FGuid& Guid, UAnimMontage& Animation);
+
+	FName FindRoleByGuid(const FGuid& Guid) const;
+
+	UAnimMontage* FindAnimationByRole(const FName& Role) const;
+
+	UAnimMontage* FindAnimationByGuid(const FGuid& Guid) const;
+
+	void Reset();
+
+private:
+
+	UPROPERTY()
+	TWeakObjectPtr<const UContextualAnimSceneAsset> SceneAssetPtr;
+
+	UPROPERTY()
+	TWeakObjectPtr<UWorld> WorldPtr;
+
+	UPROPERTY()
+	TWeakObjectPtr<ACharacter> ControlledCharacter;
 };
\ No newline at end of file
diff --git a/Engine/Plugins/Experimental/Animation/ContextualAnimation/Source/ContextualAnimationEditor/Public/ContextualAnimPreviewScene.h b/Engine/Plugins/Experimental/Animation/ContextualAnimation/Source/ContextualAnimationEditor/Public/ContextualAnimPreviewScene.h
new file mode 100644
--- /dev/null
+++ b/Engine/Plugins/Experimental/Animation/ContextualAnimation/Source/ContextualAnimationEditor/Public/ContextualAnimPreviewScene.h
@@ -0,0 +1,25 @@
+// Copyright Epic Games, Inc. All Rights Reserved.
+
+#pragma once
+
+#include "CoreMinimal.h"
+#include "AdvancedPreviewScene.h"
+
+class FContextualAnimAssetEditorToolkit;
+
+class FContextualAnimPreviewScene : public FAdvancedPreviewScene
+{
+public:
+
+	FContextualAnimPreviewScene(ConstructionValues CVS, const TSharedRef<FContextualAnimAssetEditorToolkit>& EditorToolkit);
+	~FContextualAnimPreviewScene(){}
+
+	virtual void Tick(float InDeltaTime) override;
+
+	TSharedRef<FContextualAnimAssetEditorToolkit> GetEditorToolkit() const { return EditorToolkitPtr.Pin().ToSharedRef(); }
+
+private:
+
+	/** The asset editor toolkit we are embedded in */
+	TWeakPtr<FContextualAnimAssetEditorToolkit> EditorToolkitPtr;
+};
\ No newline at end of file
diff --git a/Engine/Plugins/Experimental/Animation/ContextualAnimation/Source/ContextualAnimationEditor/Public/ContextualAnimTypeActions.h b/Engine/Plugins/Experimental/Animation/ContextualAnimation/Source/ContextualAnimationEditor/Public/ContextualAnimTypeActions.h
new file mode 100644
--- /dev/null
+++ b/Engine/Plugins/Experimental/Animation/ContextualAnimation/Source/ContextualAnimationEditor/Public/ContextualAnimTypeActions.h
@@ -0,0 +1,18 @@
+// Copyright Epic Games, Inc. All Rights Reserved.
+
+#pragma once
+
+#include "CoreMinimal.h"
+#include "AssetTypeActions_Base.h"
+
+class FContextualAnimTypeActions : public FAssetTypeActions_Base
+{
+public:
+
+	FContextualAnimTypeActions(){}
+	virtual FText GetName() const override;
+	virtual FColor GetTypeColor() const override;
+	virtual UClass* GetSupportedClass() const override;
+	virtual void OpenAssetEditor(const TArray<UObject*>& InObjects, TSharedPtr<class IToolkitHost> EditWithinLevelEditor = TSharedPtr<IToolkitHost>()) override;
+	virtual uint32 GetCategories() override;
+};
\ No newline at end of file
diff --git a/Engine/Plugins/Experimental/Animation/ContextualAnimation/Source/ContextualAnimationEditor/Public/ContextualAnimViewModel.h b/Engine/Plugins/Experimental/Animation/ContextualAnimation/Source/ContextualAnimationEditor/Public/ContextualAnimViewModel.h
new file mode 100644
--- /dev/null
+++ b/Engine/Plugins/Experimental/Animation/ContextualAnimation/Source/ContextualAnimationEditor/Public/ContextualAnimViewModel.h
@@ -0,0 +1,89 @@
+// Copyright Epic Games, Inc. All Rights Reserved.
+
+#pragma once
+
+#include "CoreMinimal.h"
+#include "UObject/GCObject.h"
+#include "MovieSceneFwd.h"
+#include "ContextualAnimTypes.h"
+#include "ContextualAnimMovieSceneSequence.h"
+#include "ISequencer.h"
+
+class FContextualAnimPreviewScene;
+class UContextualAnimSceneAsset;
+class UContextualAnimPreviewManager;
+class UContextualAnimMovieSceneSequence;
+class UMovieScene;
+class UContextualAnimMovieSceneNotifyTrack;
+class UContextualAnimMovieSceneNotifySection;
+class UAnimMontage;
+struct FMovieSceneSectionMovedParams;
+struct FNewRoleWidgetParams;
+
+class FContextualAnimViewModel : public TSharedFromThis<FContextualAnimViewModel>, public FGCObject
+{
+public:
+
+	FContextualAnimViewModel();
+	virtual ~FContextualAnimViewModel();
+
+	// ~ FGCObject interface
+	virtual void AddReferencedObjects(FReferenceCollector& Collector) override;
+	virtual FString GetReferencerName() const override { return TEXT("FContextualAnimViewModel"); }
+
+	void Initialize(UContextualAnimSceneAsset* InSceneAsset, const TSharedRef<FContextualAnimPreviewScene>& InPreviewScene);
+
+	void RefreshSequencerTracks();
+
+	TSharedPtr<ISequencer> GetSequencer();
+	TSharedPtr<FContextualAnimPreviewScene> GetPreviewScene();
+	UMovieScene* GetMovieScene() const { return MovieScene; }
+	UContextualAnimMovieSceneSequence* GetMovieSceneSequence() const { return MovieSceneSequence; }
+	UContextualAnimSceneAsset* GetSceneAsset() const { return SceneAsset; }
+	UContextualAnimPreviewManager* GetPreviewManager() const { return PreviewManager; }
+
+	void AddActorTrack(const FNewRoleWidgetParams& Params);
+
+ 	UAnimMontage* FindAnimationByGuid(const FGuid& Guid) const;
+
+	void AnimationModified(UAnimMontage& Animation);
+
+private:
+
+	/** Scene asset being viewed and edited by this view model. */
+	TObjectPtr<UContextualAnimSceneAsset> SceneAsset;
+
+	/** Manager to spawn and interact with preview actors */
+	TObjectPtr<UContextualAnimPreviewManager> PreviewManager;
+
+	/** MovieSceneSequence for displaying this scene asset in the sequencer time line. */
+	TObjectPtr<UContextualAnimMovieSceneSequence> MovieSceneSequence;
+
+	/** MovieScene for displaying this scene asset in the sequencer time line. */
+	TObjectPtr<UMovieScene> MovieScene;
+
+	/** Sequencer instance viewing and editing the scene asset */
+	TSharedPtr<ISequencer> Sequencer;
+
+	/** Weak pointer to the PreviewScene */
+	TWeakPtr<FContextualAnimPreviewScene> PreviewScenePtr;
+
+	/** The previous play status for sequencer time line. */
+	EMovieScenePlayerStatus::Type PreviousSequencerStatus;
+
+	/** The previous time for the sequencer time line. */
+	float PreviousSequencerTime = 0.f;
+
+	/** Flag for preventing OnAnimNotifyChanged from updating tracks when the change to the animation came from us */
+	bool bUpdatingAnimationFromSequencer = false;
+
+	UObject* GetPlaybackContext() const;
+
+	void SequencerTimeChanged();
+
+	void SequencerDataChanged(EMovieSceneDataChangeType DataChangeType);
+
+	void OnAnimNotifyChanged(UAnimMontage* Animation);
+
+	void CreateSequencer();
+};
\ No newline at end of file
diff --git a/Engine/Plugins/Experimental/Animation/ContextualAnimation/Source/ContextualAnimationEditor/Public/ContextualAnimViewportClient.h b/Engine/Plugins/Experimental/Animation/ContextualAnimation/Source/ContextualAnimationEditor/Public/ContextualAnimViewportClient.h
new file mode 100644
--- /dev/null
+++ b/Engine/Plugins/Experimental/Animation/ContextualAnimation/Source/ContextualAnimationEditor/Public/ContextualAnimViewportClient.h
@@ -0,0 +1,37 @@
+// Copyright Epic Games, Inc. All Rights Reserved.
+
+#pragma once
+
+#include "CoreMinimal.h"
+#include "EditorViewportClient.h"
+
+class FContextualAnimPreviewScene;
+class SContextualAnimViewport;
+class FContextualAnimAssetEditorToolkit;
+
+class FContextualAnimViewportClient : public FEditorViewportClient
+{
+public:
+
+	FContextualAnimViewportClient(const TSharedRef<FContextualAnimPreviewScene>& InPreviewScene, const TSharedRef<SContextualAnimViewport>& InViewport, const TSharedRef<FContextualAnimAssetEditorToolkit>& InAssetEditorToolkit);
+	virtual ~FContextualAnimViewportClient(){}
+
+	// ~FEditorViewportClient interface
+	virtual void Draw(const FSceneView* View, FPrimitiveDrawInterface* PDI) override;
+	virtual void TrackingStarted(const struct FInputEventState& InInputState, bool bIsDragging, bool bNudge) override;
+	virtual void TrackingStopped() override;
+	// ~End of FEditorViewportClient interface
+
+	/** Get the preview scene we are viewing */
+	TSharedRef<FContextualAnimPreviewScene> GetPreviewScene() const { return PreviewScenePtr.Pin().ToSharedRef(); }
+
+	TSharedRef<FContextualAnimAssetEditorToolkit> GetAssetEditorToolkit() const { return AssetEditorToolkitPtr.Pin().ToSharedRef(); }
+
+private:
+
+	/** Preview scene we are viewing */
+	TWeakPtr<FContextualAnimPreviewScene> PreviewScenePtr;
+
+	/** Asset editor toolkit we are embedded in */
+	TWeakPtr<FContextualAnimAssetEditorToolkit> AssetEditorToolkitPtr;
+};
diff --git a/Engine/Plugins/Experimental/Animation/ContextualAnimation/Source/ContextualAnimationEditor/Public/ContextualAnimationEditor.h b/Engine/Plugins/Experimental/Animation/ContextualAnimation/Source/ContextualAnimationEditor/Public/ContextualAnimationEditor.h
--- a/Engine/Plugins/Experimental/Animation/ContextualAnimation/Source/ContextualAnimationEditor/Public/ContextualAnimationEditor.h
+++ b/Engine/Plugins/Experimental/Animation/ContextualAnimation/Source/ContextualAnimationEditor/Public/ContextualAnimationEditor.h
@@ -5,6 +5,8 @@
 #include "CoreMinimal.h"
 #include "Modules/ModuleManager.h"
 
+class IAssetTypeActions;
+
 class FContextualAnimationEditorModule : public IModuleInterface
 {
 public:
@@ -14,4 +16,10 @@ public:
 	virtual void ShutdownModule() override;
 
 	static FContextualAnimationEditorModule& Get();
+
+private:
+
+	FDelegateHandle MovieSceneAnimNotifyTrackEditorHandle;
+
+	TSharedPtr<IAssetTypeActions> ContextualAnimAssetActions;
 };
diff --git a/Engine/Plugins/Experimental/Animation/ContextualAnimation/Source/ContextualAnimationEditor/Public/DetailCustomizations/ContextualAnimNotifySectionDetailCustom.h b/Engine/Plugins/Experimental/Animation/ContextualAnimation/Source/ContextualAnimationEditor/Public/DetailCustomizations/ContextualAnimNotifySectionDetailCustom.h
new file mode 100644
--- /dev/null
+++ b/Engine/Plugins/Experimental/Animation/ContextualAnimation/Source/ContextualAnimationEditor/Public/DetailCustomizations/ContextualAnimNotifySectionDetailCustom.h
@@ -0,0 +1,25 @@
+// Copyright Epic Games, Inc. All Rights Reserved.
+
+#pragma once
+
+#include "CoreMinimal.h"
+#include "UObject/WeakObjectPtr.h"
+#include "IDetailCustomization.h"
+
+class IDetailLayoutBuilder;
+class IDetailsView;
+class UContextualAnimMovieSceneNotifySection;
+
+/** Detail Customization for ContextualAnimMovieSceneNotifySection */
+class FContextualAnimNotifySectionDetailCustom : public IDetailCustomization
+{
+public:
+
+	static TSharedRef<IDetailCustomization> MakeInstance();
+
+	virtual void CustomizeDetails(IDetailLayoutBuilder& DetailBuilder) override;
+
+protected:
+
+	TWeakObjectPtr<UContextualAnimMovieSceneNotifySection> Section;
+};
\ No newline at end of file
diff --git a/Engine/Plugins/Experimental/Animation/ContextualAnimation/Source/ContextualAnimationEditor/Public/SContextualAnimAssetBrowser.h b/Engine/Plugins/Experimental/Animation/ContextualAnimation/Source/ContextualAnimationEditor/Public/SContextualAnimAssetBrowser.h
new file mode 100644
--- /dev/null
+++ b/Engine/Plugins/Experimental/Animation/ContextualAnimation/Source/ContextualAnimationEditor/Public/SContextualAnimAssetBrowser.h
@@ -0,0 +1,25 @@
+// Copyright Epic Games, Inc. All Rights Reserved.
+
+#pragma once
+
+#include "CoreMinimal.h"
+#include "Widgets/Layout/SBox.h"
+
+struct FAssetData;
+
+class SContextualAnimAssetBrowser : public SBox
+{
+public:
+	
+	SLATE_BEGIN_ARGS(SContextualAnimAssetBrowser) {}
+	SLATE_END_ARGS();
+
+	void Construct(const FArguments& InArgs);
+	virtual ~SContextualAnimAssetBrowser(){}
+
+private:
+
+	TSharedPtr<SBox> AssetBrowserBox;
+	void OnAssetDoubleClicked(const FAssetData& AssetData);
+	bool OnShouldFilterAsset(const FAssetData& AssetData);
+};
\ No newline at end of file
diff --git a/Engine/Plugins/Experimental/Animation/ContextualAnimation/Source/ContextualAnimationEditor/Public/SContextualAnimViewport.h b/Engine/Plugins/Experimental/Animation/ContextualAnimation/Source/ContextualAnimationEditor/Public/SContextualAnimViewport.h
new file mode 100644
--- /dev/null
+++ b/Engine/Plugins/Experimental/Animation/ContextualAnimation/Source/ContextualAnimationEditor/Public/SContextualAnimViewport.h
@@ -0,0 +1,50 @@
+// Copyright Epic Games, Inc. All Rights Reserved.
+
+#pragma once
+
+#include "CoreMinimal.h"
+#include "SEditorViewport.h"
+
+class FContextualAnimPreviewScene;
+class FContextualAnimViewportClient;
+class FContextualAnimAssetEditorToolkit;
+
+struct FContextualAnimViewportRequiredArgs
+{
+	FContextualAnimViewportRequiredArgs(const TSharedRef<FContextualAnimAssetEditorToolkit>& InAssetEditorToolkit, const TSharedRef<FContextualAnimPreviewScene>& InPreviewScene)
+		: AssetEditorToolkit(InAssetEditorToolkit)
+		, PreviewScene(InPreviewScene)
+	{}
+
+	TSharedRef<FContextualAnimAssetEditorToolkit> AssetEditorToolkit;
+
+	TSharedRef<FContextualAnimPreviewScene> PreviewScene;
+};
+
+class SContextualAnimViewport : public SEditorViewport
+{
+public:
+	
+	SLATE_BEGIN_ARGS(SContextualAnimViewport) {}
+	SLATE_END_ARGS();
+
+	void Construct(const FArguments& InArgs, const FContextualAnimViewportRequiredArgs& InRequiredArgs);
+	virtual ~SContextualAnimViewport(){}
+
+protected:
+
+	// ~SEditorViewport interface
+	virtual TSharedRef<FEditorViewportClient> MakeEditorViewportClient() override;
+	// ~End of SEditorViewport interface
+
+protected:
+	
+	/** Viewport client */
+	TSharedPtr<FContextualAnimViewportClient> ViewportClient;
+
+	/** The preview scene that we are viewing */
+	TWeakPtr<FContextualAnimPreviewScene> PreviewScenePtr;
+
+	/** Asset editor toolkit we are embedded in */
+	TWeakPtr<FContextualAnimAssetEditorToolkit> AssetEditorToolkitPtr;
+};
\ No newline at end of file
