Contextual Anim Scene Asset Editor: Fix crash when adding a section to a newly created notify track

diff --git a/Engine/Plugins/Experimental/Animation/ContextualAnimation/Source/ContextualAnimationEditor/Private/ContextualAnimMovieSceneNotifyTrackEditor.cpp b/Engine/Plugins/Experimental/Animation/ContextualAnimation/Source/ContextualAnimationEditor/Private/ContextualAnimMovieSceneNotifyTrackEditor.cpp
--- a/Engine/Plugins/Experimental/Animation/ContextualAnimation/Source/ContextualAnimationEditor/Private/ContextualAnimMovieSceneNotifyTrackEditor.cpp
+++ b/Engine/Plugins/Experimental/Animation/ContextualAnimation/Source/ContextualAnimationEditor/Private/ContextualAnimMovieSceneNotifyTrackEditor.cpp
@@ -470,19 +470,19 @@ void FContextualAnimMovieSceneNotifyTrackEditor::AddNewNotifyTrack(TArray<FGuid>
 
 		// Copied from AnimTimelineTrack_Notifies.cpp/FAnimTimelineTrack_Notifies::AddTrack()
 
-		FAnimNotifyTrack NewItem;
-		NewItem.TrackName = GetNewTrackName(Animation);
-		NewItem.TrackColor = FLinearColor::White;
+		FAnimNotifyTrack NewNotifyTrack;
+		NewNotifyTrack.TrackName = GetNewTrackName(Animation);
+		NewNotifyTrack.TrackColor = FLinearColor::White;
 
-		Animation->AnimNotifyTracks.Add(NewItem);
+		Animation->AnimNotifyTracks.Add(NewNotifyTrack);
 
 		GetMovieSceneSequence().GetViewModel().AnimationModified(*Animation);
 
-		// Create MovieSceneTrack
+		// Create and Initialize MovieSceneTrack
 		UContextualAnimMovieSceneNotifyTrack* MovieSceneTrack = GetMovieSceneSequence().GetMovieScene()->AddTrack<UContextualAnimMovieSceneNotifyTrack>(ObjectBinding);
 		check(MovieSceneTrack);
 
-		MovieSceneTrack->SetDisplayName(FText::FromString(NewItem.TrackName.ToString()));
+		MovieSceneTrack->Initialize(*Animation, NewNotifyTrack);
 	}
 
 	GetSequencer()->NotifyMovieSceneDataChanged(EMovieSceneDataChangeType::MovieSceneStructureItemAdded);
