Fix for static analysis warnings

diff --git a/Engine/Plugins/Experimental/Animation/ContextualAnimation/Source/ContextualAnimationEditor/Private/ContextualAnimMovieSceneNotifyTrack.cpp b/Engine/Plugins/Experimental/Animation/ContextualAnimation/Source/ContextualAnimationEditor/Private/ContextualAnimMovieSceneNotifyTrack.cpp
--- a/Engine/Plugins/Experimental/Animation/ContextualAnimation/Source/ContextualAnimationEditor/Private/ContextualAnimMovieSceneNotifyTrack.cpp
+++ b/Engine/Plugins/Experimental/Animation/ContextualAnimation/Source/ContextualAnimationEditor/Private/ContextualAnimMovieSceneNotifyTrack.cpp
@@ -124,15 +124,16 @@ EMovieSceneSectionMovedResult UContextualAnimMovieSceneNotifyTrack::OnSectionMov
 				Anim = NotifyEventPtr->Notify->GetTypedOuter<UAnimMontage>();
 			}
 
-			check(&GetAnimation() == Anim);
-
-			const float SectionStartTime = (float)TickResolution.AsSeconds(NotifySection->GetInclusiveStartFrame());
-			const float SectionEndTime = (float)TickResolution.AsSeconds(NotifySection->GetExclusiveEndFrame());
+			if(Anim)
+			{
+				const float SectionStartTime = (float)TickResolution.AsSeconds(NotifySection->GetInclusiveStartFrame());
+				const float SectionEndTime = (float)TickResolution.AsSeconds(NotifySection->GetExclusiveEndFrame());
 
-			NotifyEventPtr->SetTime(SectionStartTime);
-			NotifyEventPtr->SetDuration(SectionEndTime - SectionStartTime);
+				NotifyEventPtr->SetTime(SectionStartTime);
+				NotifyEventPtr->SetDuration(SectionEndTime - SectionStartTime);
 
-			GetViewModel().AnimationModified(*Anim);
+				GetViewModel().AnimationModified(*Anim);
+			}
 		}
 	}
 
diff --git a/Engine/Plugins/Experimental/Animation/ContextualAnimation/Source/ContextualAnimationEditor/Private/ContextualAnimPreviewManager.cpp b/Engine/Plugins/Experimental/Animation/ContextualAnimation/Source/ContextualAnimationEditor/Private/ContextualAnimPreviewManager.cpp
--- a/Engine/Plugins/Experimental/Animation/ContextualAnimation/Source/ContextualAnimationEditor/Private/ContextualAnimPreviewManager.cpp
+++ b/Engine/Plugins/Experimental/Animation/ContextualAnimation/Source/ContextualAnimationEditor/Private/ContextualAnimPreviewManager.cpp
@@ -207,35 +207,37 @@ void UContextualAnimPreviewManager::PreviewTimeChanged(EMovieScenePlayerStatus::
 				AnimInstance->Montage_Pause(Animation);
 				
 				MontageInstance = AnimInstance->GetActiveMontageInstance();
-				check(MontageInstance);
 			}
 		}
 
-		const float AnimPlayLength = Animation->GetPlayLength();
-		PreviousTime = FMath::Clamp(PreviousTime, 0.f, AnimPlayLength);
-		CurrentTime = FMath::Clamp(CurrentTime, 0.f, AnimPlayLength);
-
-		if (CurrentStatus == EMovieScenePlayerStatus::Stopped || CurrentStatus == EMovieScenePlayerStatus::Scrubbing)
+		if(MontageInstance)
 		{
-			PreviewActorData.ResetActorTransform(CurrentTime);
-
-			if (MontageInstance->IsPlaying())
-			{
-				MontageInstance->Pause();
-			}
+			const float AnimPlayLength = Animation->GetPlayLength();
+			PreviousTime = FMath::Clamp(PreviousTime, 0.f, AnimPlayLength);
+			CurrentTime = FMath::Clamp(CurrentTime, 0.f, AnimPlayLength);
 
-			MontageInstance->SetPosition(CurrentTime);
-		}
-		else if (CurrentStatus == EMovieScenePlayerStatus::Playing)
-		{
-			if(PlaybackSpeed > 0.f && CurrentTime < PreviousTime)
+			if (CurrentStatus == EMovieScenePlayerStatus::Stopped || CurrentStatus == EMovieScenePlayerStatus::Scrubbing)
 			{
 				PreviewActorData.ResetActorTransform(CurrentTime);
-			}
 
-			if (!MontageInstance->IsPlaying())
+				if (MontageInstance->IsPlaying())
+				{
+					MontageInstance->Pause();
+				}
+
+				MontageInstance->SetPosition(CurrentTime);
+			}
+			else if (CurrentStatus == EMovieScenePlayerStatus::Playing)
 			{
-				MontageInstance->SetPlaying(true);
+				if (PlaybackSpeed > 0.f && CurrentTime < PreviousTime)
+				{
+					PreviewActorData.ResetActorTransform(CurrentTime);
+				}
+
+				if (!MontageInstance->IsPlaying())
+				{
+					MontageInstance->SetPlaying(true);
+				}
 			}
 		}
 	}
