Fixed Unknown Cook Failure error in QAGame due to uninitialized properties

diff --git a/Engine/Plugins/Experimental/Animation/ContextualAnimation/Source/ContextualAnimationEditor/Public/ContextualAnimMovieSceneNotifyTrackEditor.h b/Engine/Plugins/Experimental/Animation/ContextualAnimation/Source/ContextualAnimationEditor/Public/ContextualAnimMovieSceneNotifyTrackEditor.h
--- a/Engine/Plugins/Experimental/Animation/ContextualAnimation/Source/ContextualAnimationEditor/Public/ContextualAnimMovieSceneNotifyTrackEditor.h
+++ b/Engine/Plugins/Experimental/Animation/ContextualAnimation/Source/ContextualAnimationEditor/Public/ContextualAnimMovieSceneNotifyTrackEditor.h
@@ -33,7 +33,7 @@ struct FNewRoleWidgetParams
 	UAnimMontage* Animation = nullptr;
 
 	UPROPERTY(EditAnywhere, Category = "Settings")
-	bool bRequiresFlyingMode;
+	bool bRequiresFlyingMode = false;
 
 	UPROPERTY(EditAnywhere, Category = "Settings")
 	FTransform MeshToComponent = FTransform(FRotator(0.f, -90.f, 0.f));
@@ -53,7 +53,7 @@ struct FNewIKTargetParams
 	FBoneReference SourceBone;
 
 	UPROPERTY(EditAnywhere, Category = "Settings")
-	EContextualAnimIKTargetProvider Provider;
+	EContextualAnimIKTargetProvider Provider = EContextualAnimIKTargetProvider::Autogenerated;
 
 	UPROPERTY(EditAnywhere, Category = "Settings", meta = (GetOptions = "GetTargetRoleOptions"))
 	FName TargetRole = NAME_None;
