Contextual Anim Editor (WIP):
 - Refactored ActorSceneComp to provide Ik Goals to Ik Rig via IIKGoalCreatorInterface
 - Added option to create and switch between variants
 - Refactored selection logic and added option in the editor to edit the trigger area in context
 - Added option to toggle draw debug of ik targets in the editor
 - Simplified alignment logic
 - Added option to simulate the interaction in the editor

diff --git a/Engine/Plugins/Experimental/Animation/ContextualAnimation/ContextualAnimation.uplugin b/Engine/Plugins/Experimental/Animation/ContextualAnimation/ContextualAnimation.uplugin
--- a/Engine/Plugins/Experimental/Animation/ContextualAnimation/ContextualAnimation.uplugin
+++ b/Engine/Plugins/Experimental/Animation/ContextualAnimation/ContextualAnimation.uplugin
@@ -31,6 +31,10 @@
 		{
 			"Name": "MotionWarping",
 			"Enabled": true
+		},
+		{
+			"Name": "IKRig",
+			"Enabled": true
 		}
 	]
 }
\ No newline at end of file
diff --git a/Engine/Plugins/Experimental/Animation/ContextualAnimation/Source/ContextualAnimation/ContextualAnimation.Build.cs b/Engine/Plugins/Experimental/Animation/ContextualAnimation/Source/ContextualAnimation/ContextualAnimation.Build.cs
--- a/Engine/Plugins/Experimental/Animation/ContextualAnimation/Source/ContextualAnimation/ContextualAnimation.Build.cs
+++ b/Engine/Plugins/Experimental/Animation/ContextualAnimation/Source/ContextualAnimation/ContextualAnimation.Build.cs
@@ -28,7 +28,9 @@ public class ContextualAnimation : ModuleRules
 				"Core",
 				"CoreUObject",
                 "Engine",
-				"MotionWarping"
+				"GameplayTags",
+				"MotionWarping",
+				"IKRig"
 			}
 			);
 			
diff --git a/Engine/Plugins/Experimental/Animation/ContextualAnimation/Source/ContextualAnimation/Private/AnimNotifyState_IKWindow.cpp b/Engine/Plugins/Experimental/Animation/ContextualAnimation/Source/ContextualAnimation/Private/AnimNotifyState_IKWindow.cpp
--- a/Engine/Plugins/Experimental/Animation/ContextualAnimation/Source/ContextualAnimation/Private/AnimNotifyState_IKWindow.cpp
+++ b/Engine/Plugins/Experimental/Animation/ContextualAnimation/Source/ContextualAnimation/Private/AnimNotifyState_IKWindow.cpp
@@ -1,6 +1,7 @@
 // Copyright Epic Games, Inc. All Rights Reserved.
 
 #include "AnimNotifyState_IKWindow.h"
+#include "Animation/AnimMontage.h"
 
 UAnimNotifyState_IKWindow::UAnimNotifyState_IKWindow(const FObjectInitializer& ObjectInitializer)
 	: Super(ObjectInitializer)
@@ -12,4 +13,46 @@ UAnimNotifyState_IKWindow::UAnimNotifyState_IKWindow(const FObjectInitializer& O
 FString UAnimNotifyState_IKWindow::GetNotifyName_Implementation() const
 {
 	return FString::Printf(TEXT("IK (%s)"), *GoalName.ToString());
+}
+
+float UAnimNotifyState_IKWindow::GetIKAlphaValue(const FName& GoalName, const FAnimMontageInstance* MontageInstance)
+{
+	if (MontageInstance && MontageInstance->Montage)
+	{
+		for (const FAnimNotifyEvent& NotifyEvent : MontageInstance->Montage->Notifies)
+		{
+			const UAnimNotifyState_IKWindow* Notify = NotifyEvent.NotifyStateClass ? Cast<const UAnimNotifyState_IKWindow>(NotifyEvent.NotifyStateClass) : nullptr;
+			if (Notify && Notify->bEnable && Notify->GoalName == GoalName)
+			{
+				const float PlayLength = MontageInstance->Montage->GetPlayLength();
+				const float StartTime = FMath::Clamp(NotifyEvent.GetTriggerTime(), 0.f, PlayLength);
+				const float EndTime = FMath::Clamp(NotifyEvent.GetEndTriggerTime(), 0.f, PlayLength);
+
+				if (FMath::IsWithinInclusive(MontageInstance->GetPosition(), StartTime, EndTime))
+				{
+					const float WindowLength = NotifyEvent.GetDuration();
+					const float LocalPos = FMath::GetMappedRangeValueClamped(FVector2D(StartTime, EndTime), FVector2D(0.f, WindowLength), MontageInstance->GetPosition());
+					const float BlendInTime = Notify->BlendIn.GetBlendTime();
+					const float BlendOutTime = Notify->BlendOut.GetBlendTime();
+					const float BlendOutStartTimeLocal = WindowLength - BlendOutTime;
+					if (LocalPos >= BlendOutStartTimeLocal)
+					{
+						const float Alpha = FMath::GetMappedRangeValueClamped(FVector2D(BlendOutStartTimeLocal, WindowLength), FVector2D(0.f, 1.f), LocalPos);
+						return 1.f - FAlphaBlend::AlphaToBlendOption(Alpha, Notify->BlendOut.GetBlendOption(), Notify->BlendOut.GetCustomCurve());
+					}
+					else if (LocalPos <= BlendInTime)
+					{
+						const float Alpha = FMath::GetMappedRangeValueClamped(FVector2D(0.f, BlendInTime), FVector2D(0.f, 1.f), LocalPos);
+						return FAlphaBlend::AlphaToBlendOption(Alpha, Notify->BlendIn.GetBlendOption(), Notify->BlendIn.GetCustomCurve());
+					}
+					else
+					{
+						return 1.f;
+					}
+				}
+			}
+		}
+	}
+
+	return 0.f;
 }
\ No newline at end of file
diff --git a/Engine/Plugins/Experimental/Animation/ContextualAnimation/Source/ContextualAnimation/Private/ContextualAnimManager.cpp b/Engine/Plugins/Experimental/Animation/ContextualAnimation/Source/ContextualAnimation/Private/ContextualAnimManager.cpp
--- a/Engine/Plugins/Experimental/Animation/ContextualAnimation/Source/ContextualAnimation/Private/ContextualAnimManager.cpp
+++ b/Engine/Plugins/Experimental/Animation/ContextualAnimation/Source/ContextualAnimation/Private/ContextualAnimManager.cpp
@@ -58,38 +58,6 @@ void UContextualAnimManager::UnregisterSceneActorComponent(UContextualAnimSceneA
 	}
 }
 
-UContextualAnimSceneActorComponent* UContextualAnimManager::FindClosestSceneActorCompToActor(const AActor* Actor) const
-{
-	SCOPE_CYCLE_COUNTER(STAT_ContextualAnim_FindClosestSceneActorComp);
-
-	//@TODO: Implement some space partitioning to reduce the cost of this search (a 2D grid might be enough)
-
-	UContextualAnimSceneActorComponent* ClosestSceneActorComp = nullptr;
-	if(Actor)
-	{
-		const FVector TestLocation = Actor->GetActorLocation();
-
-		float BestDistanceSq = MAX_FLT;
-		for (UContextualAnimSceneActorComponent* SceneActorComp : SceneActorCompContainer)
-		{
-			check(SceneActorComp);
-			check(SceneActorComp->GetOwner());
-
-			if(SceneActorComp->GetOwner() != Actor)
-			{
-				const float DistanceSq = FVector::DistSquared(SceneActorComp->GetComponentLocation(), TestLocation);
-				if (DistanceSq < BestDistanceSq)
-				{
-					BestDistanceSq = DistanceSq;
-					ClosestSceneActorComp = SceneActorComp;
-				}
-			}
-		}
-	}
-
-	return ClosestSceneActorComp;
-}
-
 void UContextualAnimManager::Tick(float DeltaTime)
 {
 	for (UContextualAnimSceneInstance* SceneInstance : Instances)
@@ -130,165 +98,126 @@ UContextualAnimSceneInstance* UContextualAnimManager::GetSceneWithActor(AActor*
 	return nullptr;
 }
 
-bool UContextualAnimManager::TryStartScene(const UContextualAnimSceneAsset* SceneAsset, const FContextualAnimSceneBindings& Bindings)
+UContextualAnimSceneInstance* UContextualAnimManager::ForceStartScene(const UContextualAnimSceneAsset& SceneAsset, const FContextualAnimStartSceneParams& Params)
 {
-	if (SceneAsset == nullptr)
-	{
-		UE_LOG(LogContextualAnim, Warning, TEXT("UContextualAnimManager::TryStartScene. Can't start scene. Reason: Invalid Scene Asset"));
-		return false;
-	}
-
-	TMap<FName, FContextualAnimSceneActorData> SceneActorMap;
-	SceneActorMap.Reserve(Bindings.RoleToActorMap.Num());
-
-	// Find primary actor first so we can perform queries relative to it
-	AActor* const* PrimaryActorPtr = Bindings.RoleToActorMap.Find(SceneAsset->PrimaryRole);
-	if (PrimaryActorPtr == nullptr)
-	{
-		UE_LOG(LogContextualAnim, Warning, TEXT("UContextualAnimManager::TryStartScene. Can't start scene. Reason: Attempting to bind an invalid actor to the primary role. SceneAsset: %s Role: %s"),
-			*GetNameSafe(SceneAsset), *SceneAsset->PrimaryRole.ToString());
-
-		return false;
-	}
-
-	AActor* PrimaryActor = *PrimaryActorPtr;
-	const UContextualAnimSceneActorComponent* PrimarySceneActorComp = PrimaryActor->FindComponentByClass<UContextualAnimSceneActorComponent>();
-	const FTransform ToWorldTransform = PrimarySceneActorComp ? PrimarySceneActorComp->GetComponentTransform() : PrimaryActor->GetActorTransform();
-
-	for (const auto& Pair : Bindings.RoleToActorMap)
+	FContextualAnimSceneBindings Bindings;
+	for (const auto& Pair : Params.RoleToActorMap)
 	{
 		FName RoleToBind = Pair.Key;
-		AActor* ActorToBind = Pair.Value;
 
+		AActor* ActorToBind = Pair.Value;
 		if (ActorToBind == nullptr)
 		{
-			UE_LOG(LogTemp, Warning, TEXT("UContextualAnimManager::TryStartScene. Can't start scene. Reason: Attempting to bind an invalid actor. SceneAsset: %s Role: %s"),
-				*GetNameSafe(SceneAsset), *RoleToBind.ToString());
+			UE_LOG(LogContextualAnim, Warning, TEXT("UContextualAnimManager::ForceStartScene. Can't start scene. Reason: Trying to bind Invalid Actor. SceneAsset: %s Role: %s"),
+				*GetNameSafe(&SceneAsset), *RoleToBind.ToString());
 
-			return false;
+			return nullptr;
 		}
 
-		// Primary role is already bound
-		if (RoleToBind != SceneAsset->PrimaryRole)
+		const FContextualAnimTrack* AnimTrack = SceneAsset.GetAnimTrack(RoleToBind, Params.VariantIdx);
+		if (AnimTrack == nullptr)
 		{
-			int32 AnimDataIndex = INDEX_NONE;
-			float AnimStartTime = 0.f;
-
-			// Attempt to use AnimDataIndex and AnimStartTime supplied with the bindings
-			if (Bindings.AnimDataIndex != INDEX_NONE)
-			{
-				if (SceneAsset->GetAnimDataForRoleAtIndex(RoleToBind, Bindings.AnimDataIndex) == nullptr)
-				{
-					UE_LOG(LogTemp, Warning, TEXT("UContextualAnimManager::TryStartScene. Can't start scene. Reason: Invalid AnimDataIndex. SceneAsset: %s Role: %s AnimDataIndex: %d"),
-						*GetNameSafe(SceneAsset), *RoleToBind.ToString(), Bindings.AnimDataIndex);
+			UE_LOG(LogContextualAnim, Warning, TEXT("UContextualAnimManager::ForceStartScene. Can't start scene. Reason: Can't find anim track for '%s'. SceneAsset: %s"), 
+				*RoleToBind.ToString(), *GetNameSafe(&SceneAsset));
 
-					return false;
-				}
-
-				AnimDataIndex = Bindings.AnimDataIndex;
-				AnimStartTime = Bindings.AnimStartTime;
-			}
-			// Query for best animation to use
-			else
-			{
-				FContextualAnimQueryResult Result;
-				const bool bQueryResult = SceneAsset->Query(RoleToBind, Result, FContextualAnimQueryParams(ActorToBind, true, true), ToWorldTransform);
-				if (bQueryResult == false)
-				{
-					UE_LOG(LogContextualAnim, Warning, TEXT("UContextualAnimManager::TryStartScene. Can't start scene. Reason: Can't find Track for actor. SceneAsset: %s Role: %s Actor: %s"),
-						*GetNameSafe(SceneAsset), *RoleToBind.ToString(), *GetNameSafe(ActorToBind));
-
-					return false;
-				}
-
-				AnimDataIndex = Result.DataIndex;
-				AnimStartTime = Result.AnimStartTime;
-			}
-
-			//@TODO: Fix this for +2 actors interactions. AnimStartTime must be calculated only once and then everyone should use it.
+			return nullptr;
+		}
 
-			SceneActorMap.Add(RoleToBind, FContextualAnimSceneActorData(ActorToBind, SceneAsset->GetAnimDataForRoleAtIndex(RoleToBind, AnimDataIndex), SceneAsset->GetTrackSettings(RoleToBind), AnimStartTime));
+		check(AnimTrack);
 
-			if (!SceneActorMap.Contains(SceneAsset->PrimaryRole))
-			{
-				SceneActorMap.Add(SceneAsset->PrimaryRole, FContextualAnimSceneActorData(PrimaryActor,
-					SceneAsset->GetAnimDataForRoleAtIndex(SceneAsset->PrimaryRole, AnimDataIndex), 
-					SceneAsset->GetTrackSettings(SceneAsset->PrimaryRole), AnimStartTime)
-				);
-			}
-		}
+		Bindings.Add(FContextualAnimSceneActorData(RoleToBind, Params.VariantIdx, *ActorToBind, *AnimTrack, Params.AnimStartTime));
 	}
 
-	UClass* Class = SceneAsset->SceneInstanceClass;
+	UClass* Class = SceneAsset.GetSceneInstanceClass();
 	UContextualAnimSceneInstance* NewInstance = Class ? NewObject<UContextualAnimSceneInstance>(this, Class) : NewObject<UContextualAnimSceneInstance>(this);
-	NewInstance->SceneAsset = SceneAsset;
-	NewInstance->SceneActorMap = MoveTemp(SceneActorMap);
+	NewInstance->SceneAsset = &SceneAsset;
+	NewInstance->Bindings = MoveTemp(Bindings);
 	NewInstance->Start();
-	NewInstance->OnSceneEnded.BindUObject(this, &UContextualAnimManager::OnSceneInstanceEnded);
+	NewInstance->OnSceneEnded.AddDynamic(this, &UContextualAnimManager::OnSceneInstanceEnded);
 
 	Instances.Add(NewInstance);
 
-	return true;
+	return NewInstance;
 }
 
-bool UContextualAnimManager::TryStartScene(const UContextualAnimSceneAsset* SceneAsset, AActor* PrimaryActor, const TSet<UContextualAnimSceneActorComponent*>& SceneActorComps)
+UContextualAnimSceneInstance* UContextualAnimManager::BP_TryStartScene(const UContextualAnimSceneAsset* SceneAsset, const FContextualAnimStartSceneParams& Params)
 {
-	if (SceneAsset == nullptr)
+	if(SceneAsset == nullptr)
 	{
 		UE_LOG(LogContextualAnim, Warning, TEXT("UContextualAnimManager::TryStartScene. Can't start scene. Reason: Invalid Scene Asset"));
-		return false;
+		return nullptr;
 	}
 
-	if (PrimaryActor == nullptr || IsActorInAnyScene(PrimaryActor))
+	return TryStartScene(*SceneAsset, Params);
+}
+
+UContextualAnimSceneInstance* UContextualAnimManager::TryStartScene(const UContextualAnimSceneAsset& SceneAsset, const FContextualAnimStartSceneParams& Params)
+{
+	// Check that we have defined a primary role in the scene asset
+	const FName PrimaryRole = SceneAsset.GetPrimaryRole();
+	if (PrimaryRole.IsNone())
 	{
-		UE_LOG(LogContextualAnim, Warning, TEXT("UContextualAnimManager::TryStartScene. Can't start scene. Reason: Invalid PrimaryActor or already in another scene. PrimaryActor: %s"), *GetNameSafe(PrimaryActor));
-		return false;
+		UE_LOG(LogContextualAnim, Warning, TEXT("UContextualAnimManager::TryStartScene. Can't start scene. Reason: Invalid Primary Role. SceneAsset: %s Role: %s"),
+			*GetNameSafe(&SceneAsset), *PrimaryRole.ToString());
+
+		return nullptr;
 	}
 
-	const TArray<FName>& Roles = SceneAsset->GetRoles();
+	// Find the actor that should be bound to the primary Role.
+	AActor* const* PrimaryActorPtr = Params.RoleToActorMap.Find(SceneAsset.GetPrimaryRole());
+	if (PrimaryActorPtr == nullptr)
+	{
+		UE_LOG(LogContextualAnim, Warning, TEXT("UContextualAnimManager::TryStartScene. Can't start scene. Reason: Can't find valid actor for the Primary Role. SceneAsset: %s Role: %s"),
+			*GetNameSafe(&SceneAsset), *PrimaryRole.ToString());
+
+		return nullptr;
+	}
 
 	FContextualAnimSceneBindings Bindings;
-	Bindings.RoleToActorMap.Add(SceneAsset->PrimaryRole, PrimaryActor);
 
-	const UContextualAnimSceneActorComponent* PrimarySceneActorComp = PrimaryActor->FindComponentByClass<UContextualAnimSceneActorComponent>();
-	const FTransform ToWorldTransform = PrimarySceneActorComp ? PrimarySceneActorComp->GetComponentTransform() : PrimaryActor->GetActorTransform();
+	FContextualAnimPrimaryActorData PrimaryActorData;
+	PrimaryActorData.Transform = (*PrimaryActorPtr)->GetActorTransform();
 
-	for (UContextualAnimSceneActorComponent* SceneActorComp : SceneActorComps)
+	const int32 VariantsNum = SceneAsset.GetTotalVariants();
+	for (int32 VariantIdx = 0; VariantIdx < VariantsNum; VariantIdx++)
 	{
-		if (SceneActorComp && SceneActorComp->GetOwner())
+		for (const auto& Pair : Params.RoleToActorMap)
 		{
-			AActor* ActorToBind = SceneActorComp->GetOwner();
-			if(PrimaryActor != ActorToBind && !IsActorInAnyScene(ActorToBind))
+			FName RoleToBind = Pair.Key;
+			AActor* ActorToBind = Pair.Value;
+
+			FContextualAnimQuerierData QuerierData;
+			QuerierData.Transform = ActorToBind->GetActorTransform();
+			QuerierData.Velocity = ActorToBind->GetVelocity();
+
+			const FContextualAnimTrack* AnimTrack = SceneAsset.GetAnimTrack(RoleToBind, VariantIdx);
+			if (AnimTrack && AnimTrack->DoesQuerierPassSelectionCriteria(PrimaryActorData, QuerierData))
+			{
+				Bindings.Add(FContextualAnimSceneActorData(RoleToBind, VariantIdx, *ActorToBind, *AnimTrack));
+			}
+			else
 			{
-				for (const FName& RoleToBind : Roles)
-				{
-					if (!Bindings.RoleToActorMap.Contains(RoleToBind))
-					{
-						FContextualAnimQueryResult Result;
-						if (SceneAsset->Query(RoleToBind, Result, FContextualAnimQueryParams(ActorToBind, true, true), ToWorldTransform))
-						{
-							Bindings.RoleToActorMap.Add(RoleToBind, ActorToBind);
-
-							if (Bindings.AnimDataIndex != INDEX_NONE)
-							{
-								Bindings.AnimDataIndex = Result.DataIndex;
-								Bindings.AnimStartTime = Result.AnimStartTime;
-							}
-
-							break;
-						}
-					}
-				}
-
-				if (Roles.Num() == Bindings.RoleToActorMap.Num())
-				{
-					return TryStartScene(SceneAsset, Bindings);
-				}
+				Bindings.Reset();
+				break;
 			}
 		}
+
+		if (Params.RoleToActorMap.Num() == Bindings.Num())
+		{
+			UClass* Class = SceneAsset.GetSceneInstanceClass();
+			UContextualAnimSceneInstance* NewInstance = Class ? NewObject<UContextualAnimSceneInstance>(this, Class) : NewObject<UContextualAnimSceneInstance>(this);
+			NewInstance->SceneAsset = &SceneAsset;
+			NewInstance->Bindings = MoveTemp(Bindings);
+			NewInstance->Start();
+			NewInstance->OnSceneEnded.AddDynamic(this, &UContextualAnimManager::OnSceneInstanceEnded);
+
+			Instances.Add(NewInstance);
+
+			return NewInstance;
+		}
 	}
 
-	return false;
+	return nullptr;
 }
 
 bool UContextualAnimManager::TryStopSceneWithActor(AActor* Actor)
diff --git a/Engine/Plugins/Experimental/Animation/ContextualAnimation/Source/ContextualAnimation/Private/ContextualAnimMetadata.cpp b/Engine/Plugins/Experimental/Animation/ContextualAnimation/Source/ContextualAnimation/Private/ContextualAnimMetadata.cpp
deleted file mode 100644
diff --git a/Engine/Plugins/Experimental/Animation/ContextualAnimation/Source/ContextualAnimation/Private/ContextualAnimSceneActorComponent.cpp b/Engine/Plugins/Experimental/Animation/ContextualAnimation/Source/ContextualAnimation/Private/ContextualAnimSceneActorComponent.cpp
--- a/Engine/Plugins/Experimental/Animation/ContextualAnimation/Source/ContextualAnimation/Private/ContextualAnimSceneActorComponent.cpp
+++ b/Engine/Plugins/Experimental/Animation/ContextualAnimation/Source/ContextualAnimation/Private/ContextualAnimSceneActorComponent.cpp
@@ -3,13 +3,15 @@
 #include "ContextualAnimSceneActorComponent.h"
 #include "Animation/AnimInstance.h"
 #include "GameFramework/Character.h"
-#include "ContextualAnimMetadata.h"
+#include "ContextualAnimSelectionCriterion.h"
 #include "Components/SkeletalMeshComponent.h"
 #include "DrawDebugHelpers.h"
 #include "ContextualAnimManager.h"
 #include "ContextualAnimSceneInstance.h"
 #include "ContextualAnimSceneAsset.h"
 #include "ContextualAnimUtilities.h"
+#include "AnimNotifyState_IKWindow.h"
+#include "EngineUtils.h"
 
 #if !(UE_BUILD_SHIPPING || UE_BUILD_TEST)
 TAutoConsoleVariable<int32> CVarContextualAnimIKDebug(TEXT("a.ContextualAnim.IK.Debug"), 0, TEXT("Draw Debug IK Targets"));
@@ -60,7 +62,7 @@ void UContextualAnimSceneActorComponent::OnJoinedScene(const FContextualAnimScen
 	SceneActorDataPtr = SceneActorData;
 
 	// Start listening to TickPose if we joined an scene where we need IK
-	if (SceneActorData->GetSettings()->IKTargetDefinitions.Num() > 0)
+	if (SceneActorData->GetIKTargetDefs().IKTargetDefs.Num() > 0)
 	{
 		if (USkeletalMeshComponent* SkelMeshComp = UContextualAnimUtilities::TryGetSkeletalMeshComponent(GetOwner()))
 		{
@@ -78,7 +80,7 @@ void UContextualAnimSceneActorComponent::OnLeftScene(const FContextualAnimSceneA
 	OnLeftSceneDelegate.Broadcast(this);
 
 	// Stop listening to TickPose if we were
-	if (SceneActorData->GetSettings()->IKTargetDefinitions.Num() > 0)
+	if (SceneActorData->GetIKTargetDefs().IKTargetDefs.Num() > 0)
 	{
 		if (USkeletalMeshComponent* SkelMeshComp = UContextualAnimUtilities::TryGetSkeletalMeshComponent(GetOwner()))
 		{
@@ -97,95 +99,108 @@ void UContextualAnimSceneActorComponent::OnTickPose(class USkinnedMeshComponent*
 
 void UContextualAnimSceneActorComponent::UpdateIKTargets()
 {
-	const UContextualAnimSceneInstance* SceneInstance = SceneActorDataPtr ? SceneActorDataPtr->GetSceneInstance() : nullptr;
-	if (SceneInstance)
-	{
-		IKTargets.Reset();
+	IKTargets.Reset();
 
-		for (const FContextualAnimIKTargetDefinition& IKTargetDef : SceneActorDataPtr->GetSettings()->IKTargetDefinitions)
-		{
-			if (UAnimInstance* AnimInstance = SceneActorDataPtr->GetAnimInstance())
-			{
+	check(SceneActorDataPtr);
+
+	const FAnimMontageInstance* MontageInstance = SceneActorDataPtr->GetAnimMontageInstance();
+	if(MontageInstance == nullptr)
+	{
+		return;
+	}
 
+	for (const FContextualAnimIKTargetDefinition& IKTargetDef : SceneActorDataPtr->GetIKTargetDefs().IKTargetDefs)
+	{
 #if !(UE_BUILD_SHIPPING || UE_BUILD_TEST)
-				const bool bDrawDebugEnable = CVarContextualAnimIKDebug.GetValueOnGameThread() > 0;
-				const float DrawDebugDuration = CVarContextualAnimIKDrawDebugLifetime.GetValueOnGameThread();
-				FTransform IKTargetParentTransformForDebug = FTransform::Identity;
+		const bool bDrawDebugEnable = CVarContextualAnimIKDebug.GetValueOnGameThread() > 0;
+		const float DrawDebugDuration = CVarContextualAnimIKDrawDebugLifetime.GetValueOnGameThread();
+		FTransform IKTargetParentTransformForDebug = FTransform::Identity;
 #endif
 
-				FTransform IKTargetTransform = FTransform::Identity;
+		FTransform IKTargetTransform = FTransform::Identity;
 
-				float Alpha = 0.f;
-				bool bCurveFound = AnimInstance->GetCurveValue(IKTargetDef.GoalName, Alpha);
+		float Alpha = UAnimNotifyState_IKWindow::GetIKAlphaValue(IKTargetDef.GoalName, MontageInstance);
 
-				// @TODO: IKTargetTransform will be off by 1 frame if we tick before target. 
-				// Should we at least add an option to the SceneAsset to setup tick dependencies or should this be entirely up to the user?
+		// @TODO: IKTargetTransform will be off by 1 frame if we tick before target. 
+		// Should we at least add an option to the SceneAsset to setup tick dependencies or should this be entirely up to the user?
 
+		if (const FContextualAnimSceneActorData* TargetSceneActorData = SceneActorDataPtr->GetSceneInstance().FindSceneActorDataByRole(IKTargetDef.TargetRoleName))
+		{
+			if (const USkeletalMeshComponent* TargetSkelMeshComp = TargetSceneActorData->GetSkeletalMeshComponent())
+			{
 				if (IKTargetDef.Provider == EContextualAnimIKTargetProvider::Autogenerated)
 				{
-					if (const FContextualAnimSceneActorData* TargetSceneActorData = SceneInstance->FindSceneActorDataForRole(IKTargetDef.TargetRoleName))
-					{
-						if (USkeletalMeshComponent* TargetSkelMeshComp = TargetSceneActorData->GetSkeletalMeshComponent())
-						{
-							const FTransform IKTargetParentTransform = TargetSkelMeshComp->GetSocketTransform(IKTargetDef.TargetBoneName);
+					const FTransform IKTargetParentTransform = TargetSkelMeshComp->GetSocketTransform(IKTargetDef.TargetBoneName);
 
-							const float Time = SceneActorDataPtr->GetAnimTime();
-							IKTargetTransform = SceneActorDataPtr->GetAnimData()->IKTargetData.ExtractTransformAtTime(IKTargetDef.GoalName, Time) * IKTargetParentTransform;
+					const float Time = MontageInstance->GetPosition();
+					IKTargetTransform = SceneActorDataPtr->GetAnimTrack().IKTargetData.ExtractTransformAtTime(IKTargetDef.GoalName, Time) * IKTargetParentTransform;
 
 #if !(UE_BUILD_SHIPPING || UE_BUILD_TEST)
-							if (bDrawDebugEnable)
-							{
-								IKTargetParentTransformForDebug = IKTargetParentTransform;
-							}
-#endif
-						}
+					if (bDrawDebugEnable)
+					{
+						IKTargetParentTransformForDebug = IKTargetParentTransform;
 					}
+#endif
 				}
 				else if (IKTargetDef.Provider == EContextualAnimIKTargetProvider::Bone)
 				{
-					if (const FContextualAnimSceneActorData* TargetSceneActorData = SceneInstance->FindSceneActorDataForRole(IKTargetDef.TargetRoleName))
-					{
-						if (USkeletalMeshComponent* TargetSkelMeshComp = TargetSceneActorData->GetSkeletalMeshComponent())
-						{
-							IKTargetTransform = TargetSkelMeshComp->GetSocketTransform(IKTargetDef.TargetBoneName);
+					IKTargetTransform = TargetSkelMeshComp->GetSocketTransform(IKTargetDef.TargetBoneName);
 
 #if !(UE_BUILD_SHIPPING || UE_BUILD_TEST)
-							if (bDrawDebugEnable)
-							{
-								IKTargetParentTransformForDebug = TargetSkelMeshComp->GetSocketTransform(TargetSkelMeshComp->GetParentBone(IKTargetDef.TargetBoneName));
-							}
-#endif
-						}
+					if (bDrawDebugEnable)
+					{
+						IKTargetParentTransformForDebug = TargetSkelMeshComp->GetSocketTransform(TargetSkelMeshComp->GetParentBone(IKTargetDef.TargetBoneName));
 					}
+#endif
 				}
+			}
+		}
 
 #if !(UE_BUILD_SHIPPING || UE_BUILD_TEST)
-				const float ForcedAlphaValue = CVarContextualAnimIKForceAlpha.GetValueOnGameThread();
-				if (ForcedAlphaValue > 0.f)
-				{
-					Alpha = FMath::Clamp(ForcedAlphaValue, 0.f, 1.f);
-				}
+		const float ForcedAlphaValue = CVarContextualAnimIKForceAlpha.GetValueOnGameThread();
+		if (ForcedAlphaValue > 0.f)
+		{
+			Alpha = FMath::Clamp(ForcedAlphaValue, 0.f, 1.f);
+		}
 
-				if (bDrawDebugEnable)
-				{
-					const float DrawThickness = 0.5f;
-					const FColor DrawColor = FColor::MakeRedToGreenColorFromScalar(Alpha);
+		if (bDrawDebugEnable)
+		{
+			const float DrawThickness = 0.5f;
+			const FColor DrawColor = FColor::MakeRedToGreenColorFromScalar(Alpha);
 
-					DrawDebugLine(GetWorld(), IKTargetParentTransformForDebug.GetLocation(), IKTargetTransform.GetLocation(), DrawColor, false, DrawDebugDuration, 0, DrawThickness);
-					DrawDebugCoordinateSystem(GetWorld(), IKTargetTransform.GetLocation(), IKTargetTransform.Rotator(), 10.f, false, DrawDebugDuration, 0, DrawThickness);
-					//DrawDebugSphere(GetWorld(), IKTargetTransform.GetLocation(), 5.f, 12, DrawColor, false, 0.f, 0, DrawThickness  );
+			DrawDebugLine(GetWorld(), IKTargetParentTransformForDebug.GetLocation(), IKTargetTransform.GetLocation(), DrawColor, false, DrawDebugDuration, 0, DrawThickness);
+			DrawDebugCoordinateSystem(GetWorld(), IKTargetTransform.GetLocation(), IKTargetTransform.Rotator(), 10.f, false, DrawDebugDuration, 0, DrawThickness);
+			//DrawDebugSphere(GetWorld(), IKTargetTransform.GetLocation(), 5.f, 12, DrawColor, false, 0.f, 0, DrawThickness  );
 
-					//DrawDebugString(GetWorld(), IKTargetTransform.GetLocation(), FString::Printf(TEXT("%s (%f)"), *IKTargetDef.AlphaCurveName.ToString(), Alpha));
-				}
+			//DrawDebugString(GetWorld(), IKTargetTransform.GetLocation(), FString::Printf(TEXT("%s (%f)"), *IKTargetDef.AlphaCurveName.ToString(), Alpha));
+		}
 #endif
 
-				// Convert IK Target to mesh space
-				IKTargetTransform = IKTargetTransform.GetRelativeTransform(SceneActorDataPtr->GetSkeletalMeshComponent()->GetComponentTransform());
+		// Convert IK Target to mesh space
+		//IKTargetTransform = IKTargetTransform.GetRelativeTransform(SceneActorDataPtr->GetSkeletalMeshComponent()->GetComponentTransform());
 
-				IKTargets.Add(FContextualAnimIKTarget(IKTargetDef.GoalName, Alpha, IKTargetTransform));	
-			}
-		}
+		IKTargets.Add(FContextualAnimIKTarget(IKTargetDef.GoalName, Alpha, IKTargetTransform));
+	}
+}
+
+void UContextualAnimSceneActorComponent::AddIKGoals_Implementation(TMap<FName, FIKRigGoal>& OutGoals)
+{
+	OutGoals.Reserve(IKTargets.Num());
+
+	for(const FContextualAnimIKTarget& IKTarget : IKTargets)
+	{
+		FIKRigGoal Goal;
+		Goal.Name = IKTarget.GoalName;
+		Goal.Position = IKTarget.Transform.GetLocation();
+		Goal.Rotation = IKTarget.Transform.Rotator();
+		Goal.PositionAlpha = IKTarget.Alpha;
+		Goal.RotationAlpha = IKTarget.Alpha;
+		Goal.PositionSpace = EIKRigGoalSpace::World;
+		Goal.RotationSpace = EIKRigGoalSpace::World;
+		OutGoals.Add(Goal.Name, Goal);
 	}
+
+	IKTargets.Reset();
 }
 
 const FContextualAnimIKTarget& UContextualAnimSceneActorComponent::GetIKTargetByGoalName(FName GoalName) const
@@ -212,49 +227,9 @@ FPrimitiveSceneProxy* UContextualAnimSceneActorComponent::CreateSceneProxy()
 		FSceneActorCompProxy(const UContextualAnimSceneActorComponent* InComponent)
 			: FPrimitiveSceneProxy(InComponent)
 			, SceneAssetPtr(InComponent->SceneAsset)
-			, Params(InComponent->DebugParams)
 		{
 		}
 
-		static void DrawSector(FPrimitiveDrawInterface* PDI, const FVector& Origin, const FVector& Direction, float MinDistance, float MaxDistance, float MinAngle, float MaxAngle, const FLinearColor& Color, uint8 DepthPriority, float Thickness)
-		{
-			// Draw Cone lines
-			const FVector LeftDirection = Direction.RotateAngleAxis(MinAngle, FVector::UpVector);
-			const FVector RightDirection = Direction.RotateAngleAxis(MaxAngle, FVector::UpVector);
-			PDI->DrawLine(Origin + (LeftDirection * MinDistance), Origin + (LeftDirection * MaxDistance), Color, DepthPriority, Thickness);
-			PDI->DrawLine(Origin + (RightDirection * MinDistance), Origin + (RightDirection * MaxDistance), Color, DepthPriority, Thickness);
-
-			// Draw Near Arc
-			FVector LastDirection = LeftDirection;
-			float Angle = MinAngle;
-			while (Angle < MaxAngle)
-			{
-				Angle = FMath::Clamp<float>(Angle + 10, MinAngle, MaxAngle);
-
-				const float Length = MinDistance;
-				const FVector NewDirection = Direction.RotateAngleAxis(Angle, FVector::UpVector);
-				const FVector LineStart = Origin + (LastDirection * Length);
-				const FVector LineEnd = Origin + (NewDirection * Length);
-				PDI->DrawLine(LineStart, LineEnd, Color, DepthPriority, Thickness);
-				LastDirection = NewDirection;
-			}
-
-			// Draw Far Arc
-			LastDirection = LeftDirection;
-			Angle = MinAngle;
-			while (Angle < MaxAngle)
-			{
-				Angle = FMath::Clamp<float>(Angle + 10, MinAngle, MaxAngle);
-
-				const float Length = MaxDistance;
-				const FVector NewDirection = Direction.RotateAngleAxis(Angle, FVector::UpVector);
-				const FVector LineStart = Origin + (LastDirection * Length);
-				const FVector LineEnd = Origin + (NewDirection * Length);
-				PDI->DrawLine(LineStart, LineEnd, Color, DepthPriority, Thickness);
-				LastDirection = NewDirection;
-			}
-		}
-
 		virtual void GetDynamicMeshElements(const TArray<const FSceneView*>& Views, const FSceneViewFamily& ViewFamily, uint32 VisibilityMap, FMeshElementCollector& Collector) const override
 		{
 			const UContextualAnimSceneAsset* Asset = SceneAssetPtr.Get();
@@ -283,99 +258,40 @@ FPrimitiveSceneProxy* UContextualAnimSceneActorComponent::CreateSceneProxy()
 
 					//DrawCircle(PDI, ToWorldTransform.GetLocation(), FVector(1, 0, 0), FVector(0, 1, 0), FColor::Red, SceneAssetPtr->GetRadius(), 12, SDPG_World, 1.f);
 
-					SceneAssetPtr->ForEachAnimData([=](const FName& Role, const FContextualAnimData& Data)
+					SceneAssetPtr->ForEachAnimTrack([=](const FContextualAnimTrack& AnimTrack)
 					{
-						if (Role != SceneAssetPtr->PrimaryRole)
+						if (AnimTrack.Role != SceneAssetPtr->GetPrimaryRole())
 						{
-							FLinearColor DrawColor = FLinearColor::White;
-
-							//@TODO: Temp, this should not be here
-							if(Params.TestActor.IsValid())
-							{
-								if (Data.Metadata)
-								{
-									const FTransform EntryTransform = Data.GetAlignmentTransformAtEntryTime() * ToWorldTransform;
-									if (Data.Metadata->DoesQuerierPassConditions(FContextualAnimQuerier(Params.TestActor.Get()), FContextualAnimQueryContext(ToWorldTransform), EntryTransform))
-									{
-										DrawColor = FLinearColor::Red;
-									}
-								}
-							}
-
 							// Draw Entry Point
-							const FTransform EntryTransform = (Data.GetAlignmentTransformAtEntryTime() * ToWorldTransform);
+							const FTransform EntryTransform = (AnimTrack.GetAlignmentTransformAtEntryTime() * ToWorldTransform);
 							DrawCoordinateSystem(PDI, EntryTransform.GetLocation(), EntryTransform.Rotator(), 20.f, SDPG_World, 3.f);
 
 							// Draw Sync Point
-							const FTransform SyncPoint = Data.GetAlignmentTransformAtSyncTime() * ToWorldTransform;
+							const FTransform SyncPoint = AnimTrack.GetAlignmentTransformAtSyncTime() * ToWorldTransform;
 							DrawCoordinateSystem(PDI, SyncPoint.GetLocation(), SyncPoint.Rotator(), 20.f, SDPG_World, 3.f);
 
-							if (Params.DrawAlignmentTransformAtTime != 0.f)
-							{
-								const FTransform RootAtTime = (Data.GetAlignmentTransformAtTime(Params.DrawAlignmentTransformAtTime) * ToWorldTransform);
-								DrawCoordinateSystem(PDI, RootAtTime.GetLocation(), RootAtTime.Rotator(), 10.f, SDPG_World, 2.f);
-							}
-
-							// Draw Facing Tolerance
-							if (Data.Metadata && Data.Metadata->Facing > 0.f)
-							{
-								DrawSector(PDI, EntryTransform.GetLocation(), EntryTransform.GetRotation().GetForwardVector(), 0.f, 30.f, -Data.Metadata->Facing, Data.Metadata->Facing, DrawColor, SDPG_World, 1.f);
-							}
-							else
-							{
-								DrawCircle(PDI, EntryTransform.GetLocation(), FVector(1, 0, 0), FVector(0, 1, 0), DrawColor, 30.f, 12, SDPG_World, 1.f);
-							}
-
-							// Draw Sector
-							if (Data.Metadata)
+							FLinearColor DrawColor = FLinearColor::White;
+							for (const UContextualAnimSelectionCriterion* Criterion : AnimTrack.SelectionCriteria)
 							{
-								FVector Origin = ToWorldTransform.GetLocation();
-								Origin.Z = EntryTransform.GetLocation().Z;
-
-								FVector Direction = (EntryTransform.GetLocation() - ToWorldTransform.GetLocation()).GetSafeNormal2D();
-
-								if (Data.Metadata->DirectionOffset != 0.f)
+								if (const UContextualAnimSelectionCriterion_TriggerArea* Spatial = Cast<UContextualAnimSelectionCriterion_TriggerArea>(Criterion))
 								{
-									Direction = Direction.RotateAngleAxis(Data.Metadata->DirectionOffset, FVector::UpVector);
-								}
-
-								if (Data.Metadata->OriginOffset.X != 0.f)
-								{
-									Origin = Origin + Direction * Data.Metadata->OriginOffset.X;
-								}
-
-								if (Data.Metadata->OriginOffset.Y != 0.f)
-								{
-									Origin = Origin + (Direction.ToOrientationQuat().GetRightVector()) * Data.Metadata->OriginOffset.Y;
-								}
-
-								if (Data.Metadata->NearWidth > 0.f || Data.Metadata->FarWidth > 0.f)
-								{
-									const float HalfNearWidth = (Data.Metadata->NearWidth / 2.f);
-									const FVector RightVector = Direction.ToOrientationQuat().GetRightVector();
-									const FVector A = Origin + (-RightVector * HalfNearWidth);
-									const FVector B = Origin + (RightVector * HalfNearWidth);
-
-									const float HalfFarWidth = (Data.Metadata->FarWidth / 2.f);
-									const FVector FarEdgeCenter = Origin + (Direction * Data.Metadata->MaxDistance);
-									const FVector C = FarEdgeCenter + (-RightVector * HalfFarWidth);
-									const FVector D = FarEdgeCenter + (RightVector * HalfFarWidth);
+									const float HalfHeight = Spatial->Height / 2.f;
+									const int32 LastIndex = Spatial->PolygonPoints.Num() - 1;
+									for (int32 Idx = 0; Idx <= LastIndex; Idx++)
+									{
+										const FVector P0 = ToWorldTransform.TransformPositionNoScale(Spatial->PolygonPoints[Idx]);
+										const FVector P1 = ToWorldTransform.TransformPositionNoScale(Spatial->PolygonPoints[Idx == LastIndex ? 0 : Idx + 1]);
 
-									PDI->DrawLine(A, B, DrawColor, SDPG_World, 2.f);
-									PDI->DrawLine(C, D, DrawColor, SDPG_World, 2.f);
+										PDI->DrawLine(P0, P1, DrawColor, SDPG_Foreground, 2.f);
+										PDI->DrawLine(P0 + FVector::UpVector * Spatial->Height, P1 + FVector::UpVector * Spatial->Height, DrawColor, SDPG_Foreground, 2.f);
 
-									PDI->DrawLine(A, C, DrawColor, SDPG_World, 2.f);
-									PDI->DrawLine(B, D, DrawColor, SDPG_World, 2.f);
-								}
-								else
-								{
-									PDI->DrawLine(Origin, Origin + Direction * Data.Metadata->MaxDistance, DrawColor, SDPG_World, 2.f);
-									DrawCircle(PDI, Origin, FVector(1, 0, 0), FVector(0, 1, 0), DrawColor, Data.Metadata->MaxDistance, 32, SDPG_World, 2.f);
+										PDI->DrawLine(P0, P0 + FVector::UpVector * Spatial->Height, DrawColor, SDPG_Foreground, 2.f);
+									}
 								}
 							}
 						}
 
-						return EContextualAnimForEachResult::Continue;
+						return UE::ContextualAnim::EForEachResult::Continue;
 					});
 				}
 			}
@@ -403,7 +319,6 @@ FPrimitiveSceneProxy* UContextualAnimSceneActorComponent::CreateSceneProxy()
 
 	private:
 		TWeakObjectPtr<const UContextualAnimSceneAsset> SceneAssetPtr;
-		FContextualAnimDebugParams Params;
 	};
 
 	if(bEnableDebug)
diff --git a/Engine/Plugins/Experimental/Animation/ContextualAnimation/Source/ContextualAnimation/Private/ContextualAnimSceneAsset.cpp b/Engine/Plugins/Experimental/Animation/ContextualAnimation/Source/ContextualAnimation/Private/ContextualAnimSceneAsset.cpp
--- a/Engine/Plugins/Experimental/Animation/ContextualAnimation/Source/ContextualAnimation/Private/ContextualAnimSceneAsset.cpp
+++ b/Engine/Plugins/Experimental/Animation/ContextualAnimation/Source/ContextualAnimation/Private/ContextualAnimSceneAsset.cpp
@@ -6,9 +6,8 @@
 #include "Animation/AnimationPoseData.h"
 #include "Animation/AnimTypes.h"
 #include "ContextualAnimUtilities.h"
-#include "ContextualAnimMetadata.h"
+#include "ContextualAnimSelectionCriterion.h"
 #include "UObject/ObjectSaveContext.h"
-#include "ContextualAnimScenePivotProvider.h"
 #include "Containers/ArrayView.h"
 
 static FCompactPoseBoneIndex GetCompactPoseBoneIndexFromPose(const FCSPose<FCompactPose>& Pose, const FName& BoneName)
@@ -31,8 +30,28 @@ UContextualAnimSceneAsset::UContextualAnimSceneAsset(const FObjectInitializer& O
 {
 	bDisableCollisionBetweenActors = true;
 	SampleRate = 15;
+
+	AlignmentSections.Add(FContextualAnimAlignmentSectionData());
 }
 
+#if WITH_EDITOR
+void UContextualAnimSceneAsset::PostEditChangeProperty(FPropertyChangedEvent& PropertyChangedEvent)
+{
+	Super::PostEditChangeProperty(PropertyChangedEvent);
+
+	const FName PropertyName = (PropertyChangedEvent.Property != nullptr) ? PropertyChangedEvent.Property->GetFName() : NAME_None;
+	const FName MemberPropertyName = (PropertyChangedEvent.MemberProperty != nullptr) ? PropertyChangedEvent.MemberProperty->GetFName() : NAME_None;
+
+	UE_LOG(LogContextualAnim, Log, TEXT("%s::PostEditChangeProperty MemberPropertyName: %s PropertyName: %s"),
+		*GetNameSafe(this), *MemberPropertyName.ToString(), *PropertyName.ToString());
+
+	if(MemberPropertyName == GET_MEMBER_NAME_CHECKED(UContextualAnimSceneAsset, AlignmentSections))
+	{
+		GenerateAlignmentTracks();
+	}
+}
+#endif
+
 void UContextualAnimSceneAsset::PreSave(FObjectPreSaveContext ObjectSaveContext)
 {
 	Super::PreSave(ObjectSaveContext);
@@ -42,285 +61,301 @@ void UContextualAnimSceneAsset::PreSave(FObjectPreSaveContext ObjectSaveContext)
 
 void UContextualAnimSceneAsset::PrecomputeData()
 {
-	// Necessary for FCompactPose that uses a FAnimStackAllocator (TMemStackAllocator) which allocates from FMemStack.
-	// When allocating memory from FMemStack we need to explicitly use FMemMark to ensure items are freed when the scope exits. 
-	// UWorld::Tick adds a FMemMark to catch any allocation inside the game tick 
-	// but any allocation from outside the game tick (like here when generating the alignment tracks off-line) must explicitly add a mark to avoid a leak 
-	FMemMark Mark(FMemStack::Get());
-
-	// Set Index for each ContextualAnimData
-	int32 NumAnimData = 0;
-	for (auto& Entry : DataContainer)
+	// Set VariantIdx for each AnimTrack
+	for (int32 VariantIdx = 0; VariantIdx < Variants.Num(); VariantIdx++)
 	{
-		NumAnimData = FMath::Max(NumAnimData, Entry.Value.AnimDataContainer.Num());
-		for (int32 AnimDataIdx = 0; AnimDataIdx < Entry.Value.AnimDataContainer.Num(); AnimDataIdx++)
+		for (FContextualAnimTrack& AnimTrack : Variants[VariantIdx].Tracks)
 		{
-			Entry.Value.AnimDataContainer[AnimDataIdx].Index = AnimDataIdx;
+			AnimTrack.VariantIdx = VariantIdx;
 		}
 	}
 
-	// Generate scene pivot for each alignment section
-	for (int32 AlignmentSectionIdx = 0; AlignmentSectionIdx < AlignmentSections.Num(); AlignmentSectionIdx++)
-	{
-		AlignmentSections[AlignmentSectionIdx].ScenePivots.Reset();
+	// Generate alignment tracks relative to scene pivot
+	GenerateAlignmentTracks();
 
-		// We need to calculate scene pivot for each set of animations
-		for (int32 AnimDataIdx = 0; AnimDataIdx < NumAnimData; AnimDataIdx++)
-		{
-			if (AlignmentSections[AlignmentSectionIdx].ScenePivotProvider)
-			{
-				const FTransform ScenePivot = AlignmentSections[AlignmentSectionIdx].ScenePivotProvider->CalculateScenePivot_Source(AnimDataIdx);
-				AlignmentSections[AlignmentSectionIdx].ScenePivots.Add(ScenePivot);
-			}
-			else
-			{
-				AlignmentSections[AlignmentSectionIdx].ScenePivots.Add(FTransform::Identity);
-			}
-		}
-	}
-
-	for (auto& Entry : DataContainer)
-	{
-		for (FContextualAnimData& Data : Entry.Value.AnimDataContainer)
-		{
-			// Generate alignment tracks relative to scene pivot
-			GenerateAlignmentTracks(Entry.Value.Settings, Data);
-
-			// Generate IK Targets
-			GenerateIKTargetTracks(Entry.Value.Settings, Data);
-		}
-	}
+	// Generate IK Targets
+	GenerateIKTargetTracks();
 
 	UpdateRadius();
 }
 
-void UContextualAnimSceneAsset::GenerateAlignmentTracks(const FContextualAnimTrackSettings& Settings, FContextualAnimData& AnimData)
+void UContextualAnimSceneAsset::GenerateAlignmentTracks()
 {
-	const FTransform MeshToComponentInverse = Settings.MeshToComponent.Inverse();
-	const float SampleInterval = 1.f / SampleRate;
-
-	// Initialize tracks for each alignment section
-	const int32 TotalTracks = AlignmentSections.Num();
-	AnimData.AlignmentData.Initialize(TotalTracks, SampleInterval);
-	for (int32 Idx = 0; Idx < AlignmentSections.Num(); Idx++)
-	{
-		AnimData.AlignmentData.Tracks.TrackNames.Add(AlignmentSections[Idx].SectionName);
-		AnimData.AlignmentData.Tracks.AnimationTracks.AddZeroed();
-	}
+	// Necessary for FCompactPose that uses a FAnimStackAllocator (TMemStackAllocator) which allocates from FMemStack.
+	// When allocating memory from FMemStack we need to explicitly use FMemMark to ensure items are freed when the scope exits. 
+	// UWorld::Tick adds a FMemMark to catch any allocation inside the game tick 
+	// but any allocation from outside the game tick (like here when generating the alignment tracks off-line) must explicitly add a mark to avoid a leak 
+	FMemMark Mark(FMemStack::Get());
 
-	if (const UAnimMontage* Animation = AnimData.Animation)
+	for (FContextualAnimTracksContainer& Variant : Variants)
 	{
-		float Time = 0.f;
-		float EndTime = Animation->GetPlayLength();
-		int32 SampleIndex = 0;
-		while (Time < EndTime)
+		// Generate scene pivot for each alignment section
+		for (const FContextualAnimAlignmentSectionData& AligmentSection : AlignmentSections)
 		{
-			Time = FMath::Clamp(SampleIndex * SampleInterval, 0.f, EndTime);
-			SampleIndex++;
-
-			const FTransform RootTransform = MeshToComponentInverse * (UContextualAnimUtilities::ExtractRootTransformFromAnimation(Animation, Time) * AnimData.MeshToScene);
+			Variant.ScenePivots.Reset();
 
-			for (int32 Idx = 0; Idx < AlignmentSections.Num(); Idx++)
+			FTransform ScenePivot = FTransform::Identity;
+			FContextualAnimTrack* AnimTrack = Variant.Tracks.FindByPredicate([&AligmentSection](const FContextualAnimTrack& AnimTrack) { return AnimTrack.Role == AligmentSection.Origin; });
+			if (AnimTrack)
 			{
-				FRawAnimSequenceTrack& SceneTrack = AnimData.AlignmentData.Tracks.AnimationTracks[Idx];
+				if (AligmentSection.bAlongClosestDistance)
+				{
+					FContextualAnimTrack* OtherAnimTrack = Variant.Tracks.FindByPredicate([&AligmentSection](const FContextualAnimTrack& AnimTrack) { return AnimTrack.Role == AligmentSection.OtherRole; });
+					if (OtherAnimTrack)
+					{
+						FTransform T1 = AnimTrack->Animation ? UContextualAnimUtilities::ExtractRootTransformFromAnimation(AnimTrack->Animation, 0.f) : FTransform::Identity;
+						T1 = (GetMeshToComponentForRole(AnimTrack->Role).Inverse() * (T1 * AnimTrack->MeshToScene));
 
-				const FTransform ScenePivotTransform = AlignmentSections[Idx].ScenePivots[AnimData.Index];
-				const FTransform RootRelativeToScenePivot = RootTransform.GetRelativeTransform(ScenePivotTransform);
+						FTransform T2 = OtherAnimTrack->Animation ? UContextualAnimUtilities::ExtractRootTransformFromAnimation(OtherAnimTrack->Animation, 0.f) : FTransform::Identity;
+						T2 = (GetMeshToComponentForRole(OtherAnimTrack->Role).Inverse() * (T2 * OtherAnimTrack->MeshToScene));
 
-				SceneTrack.PosKeys.Add(FVector3f(RootRelativeToScenePivot.GetLocation()));
-				SceneTrack.RotKeys.Add(FQuat4f(RootRelativeToScenePivot.GetRotation()));
+						ScenePivot.SetLocation(FMath::Lerp<FVector>(T1.GetLocation(), T2.GetLocation(), AligmentSection.Weight));
+						ScenePivot.SetRotation((T2.GetLocation() - T1.GetLocation()).GetSafeNormal2D().ToOrientationQuat());
+					}
+				}
+				else
+				{
+					const FTransform RootTransform = AnimTrack->Animation ? UContextualAnimUtilities::ExtractRootTransformFromAnimation(AnimTrack->Animation, 0.f) : FTransform::Identity;
+					ScenePivot = (GetMeshToComponentForRole(AnimTrack->Role).Inverse() * (RootTransform * AnimTrack->MeshToScene));
+				}
 			}
-		}
-	}
-	else
-	{
-		for (int32 Idx = 0; Idx < AlignmentSections.Num(); Idx++)
-		{
-			AnimData.AlignmentData.Tracks.TrackNames.Add(AlignmentSections[Idx].SectionName);
-			FRawAnimSequenceTrack& SceneTrack = AnimData.AlignmentData.Tracks.AnimationTracks.AddZeroed_GetRef();
 
-			const FTransform RootTransform = MeshToComponentInverse * AnimData.MeshToScene;
-			const FTransform ScenePivotTransform = AlignmentSections[Idx].ScenePivots[AnimData.Index];
-			const FTransform RootRelativeToScenePivot = RootTransform.GetRelativeTransform(ScenePivotTransform);
-
-			SceneTrack.PosKeys.Add(FVector3f(RootRelativeToScenePivot.GetLocation()));
-			SceneTrack.RotKeys.Add(FQuat4f(RootRelativeToScenePivot.GetRotation()));
+			Variant.ScenePivots.Add(ScenePivot);
 		}
-	}
-}
-
-void UContextualAnimSceneAsset::GenerateIKTargetTracks(const FContextualAnimTrackSettings& Settings, FContextualAnimData& AnimData)
-{
-	AnimData.IKTargetData.Empty();
-
-	if (Settings.IKTargetDefinitions.Num() == 0)
-	{
-		return;
-	}
 
-	if (const UAnimMontage* Animation = AnimData.Animation)
-	{
-		UE_LOG(LogContextualAnim, Log, TEXT("%s Generating IK Target Tracks. Animation: %s"), *GetNameSafe(this), *GetNameSafe(Animation));
+		// Generate alignment tracks relative to scene pivot
+		for (FContextualAnimTrack& AnimTrack : Variant.Tracks)
+		{
+			UE_LOG(LogContextualAnim, Log, TEXT("%s Generating AlignmentTracks Tracks. Animation: %s"), *GetNameSafe(this), *GetNameSafe(AnimTrack.Animation));
 
-		const float SampleInterval = 1.f / SampleRate;
+			const FTransform MeshToComponentInverse = GetMeshToComponentForRole(AnimTrack.Role).Inverse();
+			const float SampleInterval = 1.f / SampleRate;
 
-		TArray<FBoneIndexType> RequiredBoneIndexArray;
+			// Initialize tracks for each alignment section
+			const int32 TotalTracks = AlignmentSections.Num();
+			AnimTrack.AlignmentData.Initialize(TotalTracks, SampleInterval);
+			for (int32 Idx = 0; Idx < TotalTracks; Idx++)
+			{
+				AnimTrack.AlignmentData.Tracks.TrackNames.Add(AlignmentSections[Idx].WarpTargetName);
+				AnimTrack.AlignmentData.Tracks.AnimationTracks.AddZeroed();
+			}
 
-		// Helper structure to group pose extraction per target so we can extract the pose for all the bones that are relative to the same target in one pass
-		struct FPoseExtractionHelper
-		{
-			const FContextualAnimData* TargetAnimDataPtr = nullptr;
-			TArray<TTuple<FName, FName, int32, FName, int32>> BonesData; //0: GoalName, 1: MyBoneName 2: MyBoneIndex, 3: TargetBoneName, 4: TargetBoneIndex
-		};
-		TMap<FName, FPoseExtractionHelper> PoseExtractionHelperMap;
-		PoseExtractionHelperMap.Reserve(Settings.IKTargetDefinitions.Num());
-
-		int32 TotalTracks = 0;
-		for (const FContextualAnimIKTargetDefinition& IKTargetDef : Settings.IKTargetDefinitions)
-		{
-			if (IKTargetDef.Provider == EContextualAnimIKTargetProvider::Autogenerated)
+			if (const UAnimMontage* Animation = AnimTrack.Animation)
 			{
-				const FName TargetRole = IKTargetDef.TargetRoleName;
-				FPoseExtractionHelper* DataPtr = PoseExtractionHelperMap.Find(TargetRole);
-				if (DataPtr == nullptr)
+				float Time = 0.f;
+				float EndTime = Animation->GetPlayLength();
+				int32 SampleIndex = 0;
+				while (Time < EndTime)
 				{
-					// Find AnimData for target role. 
-					const FContextualAnimData* TargetAnimDataPtr = GetAnimDataForRoleAtIndex(TargetRole, AnimData.Index);
-					if (TargetAnimDataPtr == nullptr)
-					{
-						UE_LOG(LogContextualAnim, Warning, TEXT("\t Can't find AnimTrack for TargetRole '%s'"), *TargetRole.ToString());
-						continue;
-					}
-
-					DataPtr = &PoseExtractionHelperMap.Add(TargetRole);
-					DataPtr->TargetAnimDataPtr = TargetAnimDataPtr;
-				}
+					Time = FMath::Clamp(SampleIndex * SampleInterval, 0.f, EndTime);
+					SampleIndex++;
 
-				const FName BoneName = IKTargetDef.BoneName;
-				const int32 BoneIndex = Animation->GetSkeleton()->GetReferenceSkeleton().FindBoneIndex(BoneName);
-				if (BoneIndex == INDEX_NONE)
-				{
-					UE_LOG(LogContextualAnim, Warning, TEXT("\t Can't find BoneIndex. BoneName: %s Animation: %s Skel: %s"),
-						*BoneName.ToString(), *GetNameSafe(Animation), *GetNameSafe(Animation->GetSkeleton()));
+					const FTransform RootTransform = MeshToComponentInverse * (UContextualAnimUtilities::ExtractRootTransformFromAnimation(Animation, Time) * AnimTrack.MeshToScene);
 
-					continue;
-				}
+					for (int32 Idx = 0; Idx < TotalTracks; Idx++)
+					{
+						FRawAnimSequenceTrack& AlignmentTrack = AnimTrack.AlignmentData.Tracks.AnimationTracks[Idx];
 
-				// Find TargetBoneIndex. Note that we add TargetBoneIndexeven if it is INDEX_NONE. In this case, my bone will be relative to the origin of the target actor. 
-				// This is to support cases where the target actor doesn't have animation or TargetBoneName is None
-				FName TargetBoneName = IKTargetDef.TargetBoneName;
-				const UAnimMontage* TargetAnimation = DataPtr->TargetAnimDataPtr->Animation;
-				const int32 TargetBoneIndex = TargetAnimation ? TargetAnimation->GetSkeleton()->GetReferenceSkeleton().FindBoneIndex(TargetBoneName) : INDEX_NONE;
-				if (TargetBoneIndex == INDEX_NONE)
-				{
-					UE_LOG(LogContextualAnim, Log, TEXT("\t Can't find TargetBoneIndex. BoneName: %s Animation: %s Skel: %s. Track for this bone will be relative to the origin of the target role."),
-						*TargetBoneName.ToString(), *GetNameSafe(TargetAnimation), TargetAnimation ? *GetNameSafe(TargetAnimation->GetSkeleton()) : nullptr);
+						const FTransform ScenePivotTransform = Variants[AnimTrack.VariantIdx].ScenePivots[Idx];
+						const FTransform RootRelativeToScenePivot = RootTransform.GetRelativeTransform(ScenePivotTransform);
 
-					TargetBoneName = NAME_None;
+						AlignmentTrack.PosKeys.Add(FVector3f(RootRelativeToScenePivot.GetLocation()));
+						AlignmentTrack.RotKeys.Add(FQuat4f(RootRelativeToScenePivot.GetRotation()));
+					}
 				}
+			}
+			else
+			{
+				const FTransform RootTransform = MeshToComponentInverse * AnimTrack.MeshToScene;
 
-				RequiredBoneIndexArray.AddUnique(BoneIndex);
+				for (int32 Idx = 0; Idx < TotalTracks; Idx++)
+				{
+					FRawAnimSequenceTrack& SceneTrack = AnimTrack.AlignmentData.Tracks.AnimationTracks[Idx];
 
-				DataPtr->BonesData.Add(MakeTuple(IKTargetDef.GoalName, BoneName, BoneIndex, TargetBoneName, TargetBoneIndex));
-				TotalTracks++;
+					const FTransform ScenePivotTransform = Variants[AnimTrack.VariantIdx].ScenePivots[Idx];
+					const FTransform RootRelativeToScenePivot = RootTransform.GetRelativeTransform(ScenePivotTransform);
 
-				UE_LOG(LogContextualAnim, Log, TEXT("\t Bone added for extraction. GoalName: %s BoneName: %s (%d) TargetRole: %s TargetAnimation: %s TargetBone: %s (%d)"),
-					*IKTargetDef.GoalName.ToString(), *BoneName.ToString(), BoneIndex, *TargetRole.ToString(), *GetNameSafe(TargetAnimation), *TargetBoneName.ToString(), TargetBoneIndex);
+					SceneTrack.PosKeys.Add(FVector3f(RootRelativeToScenePivot.GetLocation()));
+					SceneTrack.RotKeys.Add(FQuat4f(RootRelativeToScenePivot.GetRotation()));
+				}
 			}
 		}
+	}
+}
 
-		if (TotalTracks > 0)
-		{
-			// Complete bones chain and create bone container to extract pose from my animation
-			Animation->GetSkeleton()->GetReferenceSkeleton().EnsureParentsExistAndSort(RequiredBoneIndexArray);
-			FBoneContainer BoneContainer = FBoneContainer(RequiredBoneIndexArray, FCurveEvaluationOption(false), *Animation->GetSkeleton());
+void UContextualAnimSceneAsset::GenerateIKTargetTracks()
+{
+	FMemMark Mark(FMemStack::Get());
 
-			// Initialize track container
-			AnimData.IKTargetData.Initialize(TotalTracks, SampleInterval);
+	for (FContextualAnimTracksContainer& Variant : Variants)
+	{
+		for (FContextualAnimTrack& AnimTrack : Variant.Tracks)
+		{
+			AnimTrack.IKTargetData.Empty();
 
-			// Initialize lookup map to go from track name to target role and bone (the bone this track is relative to)
-			AnimData.IKTargetTrackLookupMap.Empty(TotalTracks);
+			const FContextualAnimIKTargetDefContainer* IKTargetDefContainerPtr = RoleToIKTargetDefsMap.Find(AnimTrack.Role);
+			if (IKTargetDefContainerPtr == nullptr || IKTargetDefContainerPtr->IKTargetDefs.Num() == 0)
+			{
+				return;
+			}
 
-			float Time = 0.f;
-			float EndTime = Animation->GetPlayLength();
-			int32 SampleIndex = 0;
-			while (Time < EndTime)
+			if (const UAnimMontage* Animation = AnimTrack.Animation)
 			{
-				Time = FMath::Clamp(SampleIndex * SampleInterval, 0.f, EndTime);
-				SampleIndex++;
+				UE_LOG(LogContextualAnim, Log, TEXT("%s Generating IK Target Tracks. Animation: %s"), *GetNameSafe(this), *GetNameSafe(Animation));
+
+				const float SampleInterval = 1.f / SampleRate;
 
-				// Extract pose from my animation
-				FCSPose<FCompactPose> ComponentSpacePose;
-				UContextualAnimUtilities::ExtractComponentSpacePose(Animation, BoneContainer, Time, false, ComponentSpacePose);
+				TArray<FBoneIndexType> RequiredBoneIndexArray;
 
-				// For each target role
-				for (auto& Data : PoseExtractionHelperMap)
+				// Helper structure to group pose extraction per target so we can extract the pose for all the bones that are relative to the same target in one pass
+				struct FPoseExtractionHelper
 				{
-					// Extract pose from target animation if any
-					FCSPose<FCompactPose> OtherComponentSpacePose;
-					TArray<FBoneIndexType> OtherRequiredBoneIndexArray;
-					FBoneContainer OtherBoneContainer;
-					const UAnimMontage* OtherAnimation = Data.Value.TargetAnimDataPtr->Animation;
-					if (OtherAnimation)
+					const FContextualAnimTrack* TargetAnimTrackPtr = nullptr;
+					TArray<TTuple<FName, FName, int32, FName, int32>> BonesData; //0: GoalName, 1: MyBoneName 2: MyBoneIndex, 3: TargetBoneName, 4: TargetBoneIndex
+				};
+				TMap<FName, FPoseExtractionHelper> PoseExtractionHelperMap;
+				PoseExtractionHelperMap.Reserve(IKTargetDefContainerPtr->IKTargetDefs.Num());
+
+				int32 TotalTracks = 0;
+				for (const FContextualAnimIKTargetDefinition& IKTargetDef : IKTargetDefContainerPtr->IKTargetDefs)
+				{
+					if (IKTargetDef.Provider == EContextualAnimIKTargetProvider::Autogenerated)
 					{
-						// Prepare array with the indices of the bones to extract from target animation
-						OtherRequiredBoneIndexArray.Reserve(Data.Value.BonesData.Num());
-						for (int32 Idx = 0; Idx < Data.Value.BonesData.Num(); Idx++)
+						const FName TargetRole = IKTargetDef.TargetRoleName;
+						FPoseExtractionHelper* DataPtr = PoseExtractionHelperMap.Find(TargetRole);
+						if (DataPtr == nullptr)
 						{
-							const int32 TargetBoneIndex = Data.Value.BonesData[Idx].Get<4>();
-							if (TargetBoneIndex != INDEX_NONE)
+							// Find track for target role. 
+							const FContextualAnimTrack* TargetAnimTrackPtr = GetAnimTrack(TargetRole, AnimTrack.VariantIdx);
+							if (TargetAnimTrackPtr == nullptr)
 							{
-								OtherRequiredBoneIndexArray.AddUnique(TargetBoneIndex);
+								UE_LOG(LogContextualAnim, Warning, TEXT("\t Can't find AnimTrack for TargetRole '%s'"), *TargetRole.ToString());
+								continue;
 							}
+
+							DataPtr = &PoseExtractionHelperMap.Add(TargetRole);
+							DataPtr->TargetAnimTrackPtr = TargetAnimTrackPtr;
 						}
 
-						if (OtherRequiredBoneIndexArray.Num() > 0)
+						const FName BoneName = IKTargetDef.BoneName;
+						const int32 BoneIndex = Animation->GetSkeleton()->GetReferenceSkeleton().FindBoneIndex(BoneName);
+						if (BoneIndex == INDEX_NONE)
 						{
-							// Complete bones chain and create bone container to extract pose form the target animation
-							OtherAnimation->GetSkeleton()->GetReferenceSkeleton().EnsureParentsExistAndSort(OtherRequiredBoneIndexArray);
-							OtherBoneContainer = FBoneContainer(OtherRequiredBoneIndexArray, FCurveEvaluationOption(false), *OtherAnimation->GetSkeleton());
+							UE_LOG(LogContextualAnim, Warning, TEXT("\t Can't find BoneIndex. BoneName: %s Animation: %s Skel: %s"),
+								*BoneName.ToString(), *GetNameSafe(Animation), *GetNameSafe(Animation->GetSkeleton()));
 
-							// Extract pose from target animation
-							UContextualAnimUtilities::ExtractComponentSpacePose(OtherAnimation, OtherBoneContainer, Time, false, OtherComponentSpacePose);
+							continue;
 						}
-					}
 
-					for (int32 Idx = 0; Idx < Data.Value.BonesData.Num(); Idx++)
-					{
-						const FName TrackName = Data.Value.BonesData[Idx].Get<0>();
-						AnimData.IKTargetData.Tracks.TrackNames.Add(TrackName);
-						AnimData.IKTargetData.Tracks.AnimationTracks.AddZeroed();
-
-						// Add entry to the lookup table
-						auto& NewItem = AnimData.IKTargetTrackLookupMap.FindOrAdd(TrackName);
-						NewItem.RoleName = Data.Key;
-						NewItem.BoneName = Data.Value.BonesData[Idx].Get<3>();
-
-						// Get bone transform from my animation
-						const FName BoneName = Data.Value.BonesData[Idx].Get<1>();
-						const FCompactPoseBoneIndex BoneIndex = GetCompactPoseBoneIndexFromPose(ComponentSpacePose, BoneName);
-						const FTransform BoneTransform = (ComponentSpacePose.GetComponentSpaceTransform(BoneIndex) * AnimData.MeshToScene);
-
-						// Get bone transform from target animation
-						FTransform OtherBoneTransform = Data.Value.TargetAnimDataPtr->MeshToScene;
-						const FName TargetBoneName = Data.Value.BonesData[Idx].Get<3>();
-						if (TargetBoneName != NAME_None)
+						// Find TargetBoneIndex. Note that we add TargetBoneIndex even if it is INDEX_NONE. In this case, my bone will be relative to the origin of the target actor. 
+						// This is to support cases where the target actor doesn't have animation or TargetBoneName is None
+						FName TargetBoneName = IKTargetDef.TargetBoneName;
+						const UAnimMontage* TargetAnimation = DataPtr->TargetAnimTrackPtr->Animation;
+						const int32 TargetBoneIndex = TargetAnimation ? TargetAnimation->GetSkeleton()->GetReferenceSkeleton().FindBoneIndex(TargetBoneName) : INDEX_NONE;
+						if (TargetBoneIndex == INDEX_NONE)
 						{
-							const FCompactPoseBoneIndex OtherBoneIndex = GetCompactPoseBoneIndexFromPose(OtherComponentSpacePose, TargetBoneName);
-							OtherBoneTransform = (OtherComponentSpacePose.GetComponentSpaceTransform(OtherBoneIndex) * Data.Value.TargetAnimDataPtr->MeshToScene);
+							UE_LOG(LogContextualAnim, Log, TEXT("\t Can't find TargetBoneIndex. BoneName: %s Animation: %s Skel: %s. Track for this bone will be relative to the origin of the target role."),
+								*TargetBoneName.ToString(), *GetNameSafe(TargetAnimation), TargetAnimation ? *GetNameSafe(TargetAnimation->GetSkeleton()) : nullptr);
+
+							TargetBoneName = NAME_None;
 						}
 
-						// Get transform relative to target
-						const FTransform BoneRelativeToOther = BoneTransform.GetRelativeTransform(OtherBoneTransform);
+						RequiredBoneIndexArray.AddUnique(BoneIndex);
+
+						DataPtr->BonesData.Add(MakeTuple(IKTargetDef.GoalName, BoneName, BoneIndex, TargetBoneName, TargetBoneIndex));
+						TotalTracks++;
+
+						UE_LOG(LogContextualAnim, Log, TEXT("\t Bone added for extraction. GoalName: %s BoneName: %s (%d) TargetRole: %s TargetAnimation: %s TargetBone: %s (%d)"),
+							*IKTargetDef.GoalName.ToString(), *BoneName.ToString(), BoneIndex, *TargetRole.ToString(), *GetNameSafe(TargetAnimation), *TargetBoneName.ToString(), TargetBoneIndex);
+					}
+				}
+
+				if (TotalTracks > 0)
+				{
+					// Complete bones chain and create bone container to extract pose from my animation
+					Animation->GetSkeleton()->GetReferenceSkeleton().EnsureParentsExistAndSort(RequiredBoneIndexArray);
+					FBoneContainer BoneContainer = FBoneContainer(RequiredBoneIndexArray, FCurveEvaluationOption(false), *Animation->GetSkeleton());
+
+					// Initialize track container
+					AnimTrack.IKTargetData.Initialize(TotalTracks, SampleInterval);
 
-						// Add transform to the track
-						FRawAnimSequenceTrack& NewTrack = AnimData.IKTargetData.Tracks.AnimationTracks[Idx];
-						NewTrack.PosKeys.Add(FVector3f(BoneRelativeToOther.GetLocation()));
-						NewTrack.RotKeys.Add(FQuat4f(BoneRelativeToOther.GetRotation()));
+					float Time = 0.f;
+					float EndTime = Animation->GetPlayLength();
+					int32 SampleIndex = 0;
+					while (Time < EndTime)
+					{
+						Time = FMath::Clamp(SampleIndex * SampleInterval, 0.f, EndTime);
+						SampleIndex++;
 
-						UE_LOG(LogContextualAnim, Verbose, TEXT("\t\t Animation: %s Time: %f BoneName: %s (T: %s) Target Animation: %s TargetBoneName: %s (T: %s)"),
-							*GetNameSafe(Animation), Time, *BoneName.ToString(), *BoneTransform.GetLocation().ToString(),
-							*GetNameSafe(OtherAnimation), *TargetBoneName.ToString(), *OtherBoneTransform.GetLocation().ToString());
+						// Extract pose from my animation
+						FCSPose<FCompactPose> ComponentSpacePose;
+						UContextualAnimUtilities::ExtractComponentSpacePose(Animation, BoneContainer, Time, false, ComponentSpacePose);
+
+						// For each target role
+						for (auto& Data : PoseExtractionHelperMap)
+						{
+							// Extract pose from target animation if any
+							FCSPose<FCompactPose> OtherComponentSpacePose;
+							TArray<FBoneIndexType> OtherRequiredBoneIndexArray;
+							FBoneContainer OtherBoneContainer;
+							const UAnimMontage* OtherAnimation = Data.Value.TargetAnimTrackPtr->Animation;
+							if (OtherAnimation)
+							{
+								// Prepare array with the indices of the bones to extract from target animation
+								OtherRequiredBoneIndexArray.Reserve(Data.Value.BonesData.Num());
+								for (int32 Idx = 0; Idx < Data.Value.BonesData.Num(); Idx++)
+								{
+									const int32 TargetBoneIndex = Data.Value.BonesData[Idx].Get<4>();
+									if (TargetBoneIndex != INDEX_NONE)
+									{
+										OtherRequiredBoneIndexArray.AddUnique(TargetBoneIndex);
+									}
+								}
+
+								if (OtherRequiredBoneIndexArray.Num() > 0)
+								{
+									// Complete bones chain and create bone container to extract pose form the target animation
+									OtherAnimation->GetSkeleton()->GetReferenceSkeleton().EnsureParentsExistAndSort(OtherRequiredBoneIndexArray);
+									OtherBoneContainer = FBoneContainer(OtherRequiredBoneIndexArray, FCurveEvaluationOption(false), *OtherAnimation->GetSkeleton());
+
+									// Extract pose from target animation
+									UContextualAnimUtilities::ExtractComponentSpacePose(OtherAnimation, OtherBoneContainer, Time, false, OtherComponentSpacePose);
+								}
+							}
+
+							for (int32 Idx = 0; Idx < Data.Value.BonesData.Num(); Idx++)
+							{
+								const FName TrackName = Data.Value.BonesData[Idx].Get<0>();
+								AnimTrack.IKTargetData.Tracks.TrackNames.Add(TrackName);
+								AnimTrack.IKTargetData.Tracks.AnimationTracks.AddZeroed();
+
+								// Get bone transform from my animation
+								const FName BoneName = Data.Value.BonesData[Idx].Get<1>();
+								const FCompactPoseBoneIndex BoneIndex = GetCompactPoseBoneIndexFromPose(ComponentSpacePose, BoneName);
+								const FTransform BoneTransform = (ComponentSpacePose.GetComponentSpaceTransform(BoneIndex) * AnimTrack.MeshToScene);
+
+								// Get bone transform from target animation
+								FTransform OtherBoneTransform = Data.Value.TargetAnimTrackPtr->MeshToScene;
+								const FName TargetBoneName = Data.Value.BonesData[Idx].Get<3>();
+								if (TargetBoneName != NAME_None)
+								{
+									const FCompactPoseBoneIndex OtherBoneIndex = GetCompactPoseBoneIndexFromPose(OtherComponentSpacePose, TargetBoneName);
+									OtherBoneTransform = (OtherComponentSpacePose.GetComponentSpaceTransform(OtherBoneIndex) * Data.Value.TargetAnimTrackPtr->MeshToScene);
+								}
+
+								// Get transform relative to target
+								const FTransform BoneRelativeToOther = BoneTransform.GetRelativeTransform(OtherBoneTransform);
+
+								// Add transform to the track
+								FRawAnimSequenceTrack& NewTrack = AnimTrack.IKTargetData.Tracks.AnimationTracks[Idx];
+								NewTrack.PosKeys.Add(FVector3f(BoneRelativeToOther.GetLocation()));
+								NewTrack.RotKeys.Add(FQuat4f(BoneRelativeToOther.GetRotation()));
+
+								UE_LOG(LogContextualAnim, Verbose, TEXT("\t\t Animation: %s Time: %f BoneName: %s (T: %s) Target Animation: %s TargetBoneName: %s (T: %s)"),
+									*GetNameSafe(Animation), Time, *BoneName.ToString(), *BoneTransform.GetLocation().ToString(),
+									*GetNameSafe(OtherAnimation), *TargetBoneName.ToString(), *OtherBoneTransform.GetLocation().ToString());
+							}
+						}
 					}
 				}
 			}
@@ -331,215 +366,208 @@ void UContextualAnimSceneAsset::GenerateIKTargetTracks(const FContextualAnimTrac
 void UContextualAnimSceneAsset::UpdateRadius()
 {
 	Radius = 0.f;
-	ForEachAnimData([this](const FName& Role, const FContextualAnimData& Data)
+	ForEachAnimTrack([this](const FContextualAnimTrack& AnimTrack)
 		{
-			Radius = FMath::Max(Radius, (Data.GetAlignmentTransformAtEntryTime().GetLocation().Size()));
-			return EContextualAnimForEachResult::Continue;
+			Radius = FMath::Max(Radius, (AnimTrack.GetAlignmentTransformAtEntryTime().GetLocation().Size()));
+			return UE::ContextualAnim::EForEachResult::Continue;
 		});
 }
 
 FName UContextualAnimSceneAsset::FindRoleByAnimation(const UAnimMontage* Animation) const
 {
-	for (const auto& Pair : DataContainer)
-	{
-		for (const FContextualAnimData& Data : Pair.Value.AnimDataContainer)
+	FName Result = NAME_None;
+
+	ForEachAnimTrack([&Result, Animation](const FContextualAnimTrack& AnimTrack)
 		{
-			if (Data.Animation == Animation)
+			if(AnimTrack.Animation == Animation)
 			{
-				return Pair.Key;
+				Result = AnimTrack.Role;
+				return UE::ContextualAnim::EForEachResult::Break;
 			}
-		}
-	}
 
-	return NAME_None;
+			return UE::ContextualAnim::EForEachResult::Continue;
+		});
+
+	return Result;
 }
 
-TArray<FName> UContextualAnimSceneAsset::GetIKGoalsForRole(const FName& Role) const
+const FContextualAnimIKTargetDefContainer& UContextualAnimSceneAsset::GetIKTargetDefsForRole(const FName& Role) const
 {
-	TArray<FName> Goals;
+	if (const FContextualAnimIKTargetDefContainer* ContainerPtr = RoleToIKTargetDefsMap.Find(Role))
+	{
+		return *ContainerPtr;
+	}
+
+	return FContextualAnimIKTargetDefContainer::EmptyContainer;
+}
 
-	if(const FContextualAnimTrackSettings* Settings = GetTrackSettings(Role))
+const FTransform& UContextualAnimSceneAsset::GetMeshToComponentForRole(const FName& Role) const
+{
+	if (RolesAsset)
 	{
-		Goals.Reserve(Settings->IKTargetDefinitions.Num());
-		for(const FContextualAnimIKTargetDefinition& IKTargetDef : Settings->IKTargetDefinitions)
+		if (const FContextualAnimRoleDefinition* RoleDef = RolesAsset->FindRoleDefinitionByName(Role))
 		{
-			Goals.Add(IKTargetDef.GoalName);
+			return RoleDef->MeshToComponent;
 		}
 	}
 
-	return Goals;
+	return FTransform::Identity;
 }
 
-bool UContextualAnimSceneAsset::QueryCompositeTrack(const FContextualAnimCompositeTrack* Track, FContextualAnimQueryResult& OutResult, const FContextualAnimQueryParams& QueryParams, const FTransform& ToWorldTransform) const
+const FContextualAnimTrack* UContextualAnimSceneAsset::GetAnimTrack(const FName& Role, int32 VariantIdx) const
 {
-	OutResult.Reset();
-
-	const FTransform QueryTransform = QueryParams.Querier.IsValid() ? QueryParams.Querier->GetActorTransform() : QueryParams.QueryTransform;
-
-	int32 DataIndex = INDEX_NONE;
-	if (Track)
+	if(Variants.IsValidIndex(VariantIdx))
 	{
-		if (QueryParams.bComplexQuery)
-		{
-			for (int32 Idx = 0; Idx < Track->AnimDataContainer.Num(); Idx++)
-			{
-				const FContextualAnimData& Data = Track->AnimDataContainer[Idx];
-
-				if (Data.Metadata)
-				{
-					const FTransform EntryTransform = Data.GetAlignmentTransformAtEntryTime() * ToWorldTransform;
+		return Variants[VariantIdx].Tracks.FindByPredicate([Role](const FContextualAnimTrack& AnimTrack){ return AnimTrack.Role == Role;});
+	}
 
-					if (!Data.Metadata->DoesQuerierPassConditions(FContextualAnimQuerier(QueryTransform), FContextualAnimQueryContext(ToWorldTransform), EntryTransform))
-					{
-						continue;
-					}
-				}
+	return nullptr;
+}
 
-				// Return the first item that passes all tests
-				DataIndex = Idx;
-				break;
-			}
-		}
-		else // Simple Query
+void UContextualAnimSceneAsset::ForEachAnimTrack(FForEachAnimTrackFunction Function) const
+{
+	for (const FContextualAnimTracksContainer& TracksContainer : Variants)
+	{
+		for (const FContextualAnimTrack& AnimTrack : TracksContainer.Tracks)
 		{
-			float BestDistanceSq = MAX_FLT;
-			for (int32 Idx = 0; Idx < Track->AnimDataContainer.Num(); Idx++)
+			if (Function(AnimTrack) == UE::ContextualAnim::EForEachResult::Break)
 			{
-				const FContextualAnimData& Data = Track->AnimDataContainer[Idx];
-
-				//@TODO: Convert querier location to local space instead
-				const FTransform EntryTransform = Data.GetAlignmentTransformAtEntryTime() * ToWorldTransform;
-				const float DistSq = FVector::DistSquared2D(EntryTransform.GetLocation(), QueryTransform.GetLocation());
-				if (DistSq < BestDistanceSq)
-				{
-					BestDistanceSq = DistSq;
-					DataIndex = Idx;
-				}
+				return;
 			}
 		}
+	}
+}
 
-		if (DataIndex != INDEX_NONE)
+void UContextualAnimSceneAsset::ForEachAnimTrack(int32 VariantIdx, FForEachAnimTrackFunction Function) const
+{
+	if(Variants.IsValidIndex(VariantIdx))
+	{
+		const TArray<FContextualAnimTrack>& Tracks = Variants[VariantIdx].Tracks;
+		for (const FContextualAnimTrack& AnimTrack : Tracks)
 		{
-			const FContextualAnimData& ResultData = Track->AnimDataContainer[DataIndex];
-
-			OutResult.DataIndex = DataIndex;
-			OutResult.Animation = ResultData.Animation;
-			OutResult.EntryTransform = ResultData.GetAlignmentTransformAtEntryTime() * ToWorldTransform;
-			OutResult.SyncTransform = ResultData.GetAlignmentTransformAtSyncTime() * ToWorldTransform;
-
-			if (QueryParams.bFindAnimStartTime)
+			if (Function(AnimTrack) == UE::ContextualAnim::EForEachResult::Break)
 			{
-				const FVector LocalLocation = (QueryTransform.GetRelativeTransform(ToWorldTransform)).GetLocation();
-				OutResult.AnimStartTime = ResultData.FindBestAnimStartTime(LocalLocation);
+				return;
 			}
-
-			return true;
 		}
 	}
-
-	return false;
 }
 
-const FContextualAnimTrackSettings* UContextualAnimSceneAsset::GetTrackSettings(const FName& Role) const
+TArray<FName> UContextualAnimSceneAsset::GetRoles() const
 {
-	const FContextualAnimCompositeTrack* Track = DataContainer.Find(Role);
-	return Track ? &Track->Settings : nullptr;
-}
+ 	TArray<FName> Result;
 
-const FContextualAnimData* UContextualAnimSceneAsset::GetAnimDataForRoleAtIndex(const FName& Role, int32 Index) const
-{
-	if(const FContextualAnimCompositeTrack* Track = DataContainer.Find(Role))
+	if(RolesAsset)
 	{
-		if(Track->AnimDataContainer.IsValidIndex(Index))
+		for (const FContextualAnimRoleDefinition& RoleDef : RolesAsset->Roles)
 		{
-			return &Track->AnimDataContainer[Index];
+			Result.Add(RoleDef.Name);
 		}
 	}
 
-	return nullptr;
+ 	return Result;
 }
 
-void UContextualAnimSceneAsset::ForEachAnimData(FForEachAnimDataFunction Function) const
+const FContextualAnimTrack* UContextualAnimSceneAsset::FindFirstAnimTrackForRoleThatPassesSelectionCriteria(const FName& Role, const FContextualAnimPrimaryActorData& PrimaryActorData, const FContextualAnimQuerierData& QuerierData) const
 {
-	for (const auto& Pair : DataContainer)
+	for (const FContextualAnimTracksContainer& Variant : Variants)
 	{
-		for (const FContextualAnimData& Data : Pair.Value.AnimDataContainer)
+		for (const FContextualAnimTrack& AnimTrack : Variant.Tracks)
 		{
-			if (Function(Pair.Key, Data) == EContextualAnimForEachResult::Break)
+			if (AnimTrack.Role == Role)
 			{
-				return;
+				bool bSuccess = true;
+				for (const UContextualAnimSelectionCriterion* Criterion : AnimTrack.SelectionCriteria)
+				{
+					if (Criterion && !Criterion->DoesQuerierPassCondition(PrimaryActorData, QuerierData))
+					{
+						bSuccess = false;
+						break;
+					}
+				}
+
+				if (bSuccess)
+				{
+					return &AnimTrack;
+				}
 			}
 		}
 	}
-}
 
-TArray<FName> UContextualAnimSceneAsset::GetRoles() const
-{
-	TArray<FName> Roles;
-	DataContainer.GetKeys(Roles);
-	return Roles;
+	return nullptr;
 }
 
-bool UContextualAnimSceneAsset::Query(const FName& Role, FContextualAnimQueryResult& OutResult, const FContextualAnimQueryParams& QueryParams, const FTransform& ToWorldTransform) const
+const FContextualAnimTrack* UContextualAnimSceneAsset::FindAnimTrackForRoleWithClosestEntryLocation(const FName& Role, const FContextualAnimPrimaryActorData& PrimaryActorData, const FVector& TestLocation) const
 {
-	const FContextualAnimCompositeTrack* Track = DataContainer.Find(Role);
-	return Track ? QueryCompositeTrack(Track, OutResult, QueryParams, ToWorldTransform) : false;
-}
+	const FContextualAnimTrack* Result = nullptr;
 
-UAnimMontage* UContextualAnimSceneAsset::GetAnimationForRoleAtIndex(FName Role, int32 Index) const
-{
-	if(const FContextualAnimCompositeTrack* Track = DataContainer.Find(Role))
+	float BestDistanceSq = MAX_FLT;
+	for (const FContextualAnimTracksContainer& Variant : Variants)
 	{
-		if(Track->AnimDataContainer.IsValidIndex(Index))
+		for (const FContextualAnimTrack& AnimTrack : Variant.Tracks)
 		{
-			return Track->AnimDataContainer[Index].Animation;
+			if (AnimTrack.Role == Role)
+			{
+				const FTransform EntryTransform = AnimTrack.GetAlignmentTransformAtEntryTime() * PrimaryActorData.Transform;
+				const float DistSq = FVector::DistSquared(EntryTransform.GetLocation(), TestLocation);
+				if (DistSq < BestDistanceSq)
+				{
+					BestDistanceSq = DistSq;
+					Result = &AnimTrack;
+				}
+			}
 		}
 	}
 
-	return nullptr;
+	return Result;
 }
 
-FTransform UContextualAnimSceneAsset::ExtractAlignmentTransformAtTime(FName Role, int32 AnimDataIndex, float Time) const
+FTransform UContextualAnimSceneAsset::GetAlignmentTransformForRoleRelativeToScenePivot(FName Role, int32 VariantIdx, float Time) const
 {
-	if (const FContextualAnimCompositeTrack* Track = DataContainer.Find(Role))
+	if (const FContextualAnimTrack* AnimTrack = GetAnimTrack(Role, VariantIdx))
 	{
-		if (Track->AnimDataContainer.IsValidIndex(AnimDataIndex))
-		{
-			return Track->AnimDataContainer[AnimDataIndex].GetAlignmentTransformAtTime(Time);
-		}
+		return AnimTrack->GetAlignmentTransformAtTime(Time);
 	}
 
 	return FTransform::Identity;
 }
 
-FTransform UContextualAnimSceneAsset::ExtractIKTargetTransformAtTime(FName Role, int32 AnimDataIndex, FName TrackName, float Time) const
+FTransform UContextualAnimSceneAsset::GetAlignmentTransformForRoleRelativeToOtherRole(FName FromRole, FName ToRole, int32 VariantIdx, float Time) const
+{
+	const FTransform FromRoleRelativeToScenePivot = GetAlignmentTransformForRoleRelativeToScenePivot(FromRole, VariantIdx, Time);
+	const FTransform ToRoleRelativeToScenePivot = GetAlignmentTransformForRoleRelativeToScenePivot(ToRole, VariantIdx, Time);
+	return FromRoleRelativeToScenePivot.GetRelativeTransform(ToRoleRelativeToScenePivot);
+}
+
+FTransform UContextualAnimSceneAsset::GetIKTargetTransformForRoleAtTime(FName Role, int32 VariantIdx, FName TrackName, float Time) const
 {
-	if (const FContextualAnimCompositeTrack* Track = DataContainer.Find(Role))
+	if(const FContextualAnimTrack* AnimTrack = GetAnimTrack(Role, VariantIdx))
 	{
-		if (Track->AnimDataContainer.IsValidIndex(AnimDataIndex))
-		{
-			return Track->AnimDataContainer[AnimDataIndex].IKTargetData.ExtractTransformAtTime(TrackName, Time);
-		}
+		return AnimTrack->IKTargetData.ExtractTransformAtTime(TrackName, Time);		
 	}
 
 	return FTransform::Identity;
 }
 
-int32 UContextualAnimSceneAsset::FindAnimIndex(FName Role, UAnimMontage* Animation) const
+int32 UContextualAnimSceneAsset::FindVariantIdx(FName Role, UAnimMontage* Animation) const
 {
 	int32 Result = INDEX_NONE;
 
-	if (const FContextualAnimCompositeTrack* Track = DataContainer.Find(Role))
-	{
-		for(int32 Index = 0; Index < Track->AnimDataContainer.Num(); Index++)
+	ForEachAnimTrack([&Result, Role, Animation](const FContextualAnimTrack& AnimTrack)
 		{
-			if(Track->AnimDataContainer[Index].Animation == Animation)
+			if (AnimTrack.Role == Role && AnimTrack.Animation == Animation)
 			{
-				Result = Index;
-				break;
+				Result = AnimTrack.VariantIdx;
+				return UE::ContextualAnim::EForEachResult::Break;
 			}
-		}
-	}
+
+			return UE::ContextualAnim::EForEachResult::Continue;
+		});
 
 	return Result;
+}
+
+int32 UContextualAnimSceneAsset::GetTotalVariants() const
+{
+	return Variants.Num();
 }
\ No newline at end of file
diff --git a/Engine/Plugins/Experimental/Animation/ContextualAnimation/Source/ContextualAnimation/Private/ContextualAnimSceneInstance.cpp b/Engine/Plugins/Experimental/Animation/ContextualAnimation/Source/ContextualAnimation/Private/ContextualAnimSceneInstance.cpp
--- a/Engine/Plugins/Experimental/Animation/ContextualAnimation/Source/ContextualAnimation/Private/ContextualAnimSceneInstance.cpp
+++ b/Engine/Plugins/Experimental/Animation/ContextualAnimation/Source/ContextualAnimation/Private/ContextualAnimSceneInstance.cpp
@@ -8,10 +8,8 @@
 #include "GameFramework/CharacterMovementComponent.h"
 #include "MotionWarpingComponent.h"
 #include "ContextualAnimActorInterface.h"
-#include "ContextualAnimScenePivotProvider.h"
 #include "ContextualAnimSceneAsset.h"
 #include "ContextualAnimSceneActorComponent.h"
-#include "ContextualAnimTransition.h"
 #include "ContextualAnimUtilities.h"
 #include "Engine/World.h"
 #include "DrawDebugHelpers.h"
@@ -40,7 +38,7 @@ USkeletalMeshComponent* FContextualAnimSceneActorData::GetSkeletalMeshComponent(
 
 FAnimMontageInstance* FContextualAnimSceneActorData::GetAnimMontageInstance() const
 {
-	if (const UAnimMontage* Animation = AnimDataPtr->Animation)
+	if (const UAnimMontage* Animation = AnimTrackPtr->Animation)
 	{
 		if (UAnimInstance* AnimInstance = GetAnimInstance())
 		{
@@ -51,7 +49,7 @@ FAnimMontageInstance* FContextualAnimSceneActorData::GetAnimMontageInstance() co
 	return nullptr;
 }
 
-const UAnimMontage* FContextualAnimSceneActorData::GetAnimMontage() const
+UAnimMontage* FContextualAnimSceneActorData::GetAnimMontage() const
 {
 	const FAnimMontageInstance* MontageInstance = GetAnimMontageInstance();
 	return MontageInstance ? MontageInstance->Montage : nullptr;
@@ -80,6 +78,11 @@ int32 FContextualAnimSceneActorData::GetCurrentSectionIndex() const
 	return INDEX_NONE;
 }
 
+const FContextualAnimIKTargetDefContainer& FContextualAnimSceneActorData::GetIKTargetDefs() const
+{ 
+	return GetSceneInstance().GetSceneAsset().GetIKTargetDefsForRole(Role); 
+}
+
 //================================================================================================================
 
 void UContextualAnimSceneInstance::BreakContextualAnimSceneActorData(const FContextualAnimSceneActorData& SceneActorData, AActor*& Actor, UAnimMontage*& Montage, float& AnimTime, int32& CurrentSectionIndex, FName& CurrentSectionName)
@@ -106,131 +109,16 @@ UWorld* UContextualAnimSceneInstance::GetWorld()const
 
 void UContextualAnimSceneInstance::Tick(float DeltaTime)
 {
-	UpdateTransitions(DeltaTime);
-}
-
-float UContextualAnimSceneInstance::GetResumePositionForSceneActor_Implementation(const FContextualAnimSceneActorData& SceneActorData, int32 DesiredSectionIndex) const
-{
-	// By default we resume from the same position in the section the leader is (we expect sections to have same length)
-
-	float StartTime = 0.f;
-	float EndTime = 0.f;
-	SceneActorData.GetAnimMontage()->GetSectionStartAndEndTime(DesiredSectionIndex, StartTime, EndTime);
-
-	const float TimeMaster = GetPositionInCurrentSection();
-	const float ResumePosition = StartTime + TimeMaster;
-
-	UE_LOG(LogContextualAnim, Verbose, TEXT("UContextualAnimSceneInstance::GetResumePositionForSceneActor Anim: %s DesiredSectionIndex: %d [%f %f] TimeMaster: %f ResumePosition: %f"),
-		*GetNameSafe(SceneActorData.GetAnimMontage()), DesiredSectionIndex, StartTime, EndTime, TimeMaster, ResumePosition);
-
-	return ResumePosition;
-}
-
-void UContextualAnimSceneInstance::UpdateTransitions(float DeltaTime)
-{
-	if (SceneAsset && SceneAsset->Transitions.Num() > 0)
-	{
-		FName CurrentSection = NAME_None;
-
-		if (const FContextualAnimSceneActorData* LeaderSceneActorData = FindSceneActorDataForRole(SceneAsset->GetLeaderRole()))
-		{
-			const FAnimMontageInstance* LeaderMontageInstance = LeaderSceneActorData->GetAnimMontageInstance();
-			if(LeaderMontageInstance)
-			{
-				CurrentSection = LeaderMontageInstance->GetCurrentSection();
-
-				// Attempt to resume montages that has been paused due a failed transition
-				for (auto& Pair : SceneActorMap)
-				{
-					FAnimMontageInstance* MontageInstance = Pair.Value.GetAnimMontageInstance();
-					if (MontageInstance && !MontageInstance->IsPlaying())
-					{
-						const int32 DesiredSectionIndex = MontageInstance->Montage->GetSectionIndex(CurrentSection);
-						if (DesiredSectionIndex != INDEX_NONE)
-						{
-							const float Position = GetResumePositionForSceneActor(Pair.Value, DesiredSectionIndex);
-							MontageInstance->SetPosition(Position);
-							MontageInstance->SetPlaying(true);
-						}
-					}
-				}
-
-				// Attempt to transition from current section
-				for (const FContextualAnimTransitionContainer& TransitionData : SceneAsset->Transitions)
-				{
-					if (TransitionData.FromSections.Contains(CurrentSection))
-					{
-						//@TODO: bForceTransition should not be in the SceneAsset
-						const bool bCanEnterTransition = TransitionData.bForceTransition || (TransitionData.Transition && TransitionData.Transition->CanEnterTransition(this, CurrentSection, TransitionData.ToSection));
-						if (bCanEnterTransition)
-						{
-							for (auto& Pair : SceneActorMap)
-							{
-								const bool bTransitionToResult = TransitionTo(Pair.Value, TransitionData.ToSection);
-
-								// If transition failed but we have a valid MontageInstance is usually because this montage doesn't have the desired section 
-								// In that case, just pause the montage. We will resume as soon as we have a valid section (see above)
-								if (bTransitionToResult == false)
-								{
-									if (FAnimMontageInstance* AnimMontageInstance = Pair.Value.GetAnimMontageInstance())
-									{
-										AnimMontageInstance->Pause();
-									}
-								}
-							}
-
-							// Break after find the first valid transition
-							break;
-						}
-					}
-				}
-			}
-		}
-	}
 }
 
 bool UContextualAnimSceneInstance::IsActorInThisScene(const AActor* Actor) const
 {
-	if(Actor)
-	{
-		for (const auto& Pair : SceneActorMap)
-		{
-			const FContextualAnimSceneActorData& Data = Pair.Value;
-			if(Data.GetActor() == Actor)
-			{
-				return true;
-			}
-		}
-	}
-
-	return false;
-}
-
-const FContextualAnimSceneActorData* UContextualAnimSceneInstance::FindSceneActorDataForActor(const AActor* Actor) const
-{
-	if (Actor)
-	{
-		for (const auto& Pair : SceneActorMap)
-		{
-			const FContextualAnimSceneActorData& Data = Pair.Value;
-			if (Data.GetActor() == Actor)
-			{
-				return &Data;
-			}
-		}
-	}
-
-	return nullptr;
-}
-
-const FContextualAnimSceneActorData* UContextualAnimSceneInstance::FindSceneActorDataForRole(const FName& Role) const
-{
-	return SceneActorMap.Find(Role);
+	return FindSceneActorDataByActor(Actor) != nullptr;
 }
 
 AActor* UContextualAnimSceneInstance::GetActorByRole(FName Role) const
 {
-	const FContextualAnimSceneActorData* SceneActorData = FindSceneActorDataForRole(Role);
+	const FContextualAnimSceneActorData* SceneActorData = FindSceneActorDataByRole(Role);
 	return SceneActorData ? SceneActorData->GetActor() : nullptr;
 }
 
@@ -239,25 +127,14 @@ void UContextualAnimSceneInstance::Join(FContextualAnimSceneActorData& SceneActo
 	AActor* Actor = SceneActorData.GetActor();
 	if (UAnimInstance* AnimInstance = SceneActorData.GetAnimInstance())
 	{
-		if (UMotionWarpingComponent* MotionWarpingComp = Actor->FindComponentByClass<UMotionWarpingComponent>())
-		{
-			for (const auto& Pair : AlignmentSectionToScenePivotList)
-			{
-				const FName SyncPointName = Pair.Key;
-				const float SyncTime = SceneActorData.GetAnimData()->GetSyncTimeForWarpSection(SyncPointName);
-				const FTransform AlignmentTransform = (SceneActorData.GetAnimData()->AlignmentData.ExtractTransformAtTime(SyncPointName, SyncTime) * Pair.Value);
-				MotionWarpingComp->AddOrUpdateWarpTargetFromTransform(SyncPointName, AlignmentTransform);
-			}
-		}
-
-		UAnimMontage* Animation = SceneActorData.GetAnimData()->Animation;
+		UAnimMontage* Animation = SceneActorData.GetAnimTrack().Animation;
 		AnimInstance->Montage_Play(Animation, 1.f, EMontagePlayReturnType::MontageLength, SceneActorData.GetAnimStartTime());
 
 		AnimInstance->OnPlayMontageNotifyBegin.AddUniqueDynamic(this, &UContextualAnimSceneInstance::OnNotifyBeginReceived);
 		AnimInstance->OnPlayMontageNotifyEnd.AddUniqueDynamic(this, &UContextualAnimSceneInstance::OnNotifyEndReceived);
 		AnimInstance->OnMontageBlendingOut.AddUniqueDynamic(this, &UContextualAnimSceneInstance::OnMontageBlendingOut);
 
-		if(SceneActorData.GetAnimData()->bRequireFlyingMode)
+		if(SceneActorData.GetAnimTrack().bRequireFlyingMode)
 		{
 			if (UCharacterMovementComponent* CharacterMovementComp = Actor->FindComponentByClass<UCharacterMovementComponent>())
 			{
@@ -266,7 +143,21 @@ void UContextualAnimSceneInstance::Join(FContextualAnimSceneActorData& SceneActo
 		}
 	}
 
-	if(SceneAsset->bDisableCollisionBetweenActors)
+	if (UMotionWarpingComponent* MotionWarpComp = Actor->FindComponentByClass<UMotionWarpingComponent>())
+	{
+		for (const auto& Pair : AlignmentSectionToScenePivotList)
+		{
+			const FName WarpTargetName = Pair.Key;
+			const FTransform ScenePivotRuntime = Pair.Value;
+			const float Time = SceneActorData.GetAnimTrack().GetSyncTimeForWarpSection(WarpTargetName);
+			const FTransform TransformRelativeToScenePivot = SceneActorData.GetAnimTrack().AlignmentData.ExtractTransformAtTime(WarpTargetName, Time);
+			const FTransform WarpTarget = (TransformRelativeToScenePivot * ScenePivotRuntime);
+
+			MotionWarpComp->AddOrUpdateWarpTargetFromTransform(WarpTargetName, WarpTarget);
+		}
+	}
+
+	if(SceneAsset->GetDisableCollisionBetweenActors())
 	{
 		SetIgnoreCollisionWithOtherActors(Actor, true);
 	}
@@ -278,7 +169,7 @@ void UContextualAnimSceneInstance::Join(FContextualAnimSceneActorData& SceneActo
 		SceneActorComp->OnJoinedScene(&SceneActorData);
 	}
 
-	OnActorJoined.ExecuteIfBound(this, Actor);
+	OnActorJoined.Broadcast(this, Actor);
 }
 
 void UContextualAnimSceneInstance::Leave(FContextualAnimSceneActorData& SceneActorData)
@@ -291,7 +182,7 @@ void UContextualAnimSceneInstance::Leave(FContextualAnimSceneActorData& SceneAct
 		if (UAnimInstance* AnimInstance = SceneActorData.GetAnimInstance())
 		{
 			UAnimMontage* CurrentMontage = AnimInstance->GetCurrentActiveMontage();
-			if (ensureAlways(CurrentMontage))
+			if (CurrentMontage)
 			{
 				AnimInstance->Montage_Stop(CurrentMontage->BlendOut.GetBlendTime(), CurrentMontage);
 			}
@@ -341,45 +232,43 @@ bool UContextualAnimSceneInstance::TransitionTo(FContextualAnimSceneActorData& S
 
 void UContextualAnimSceneInstance::Start()
 {
-	const int32 AlignmentSectionsNum = SceneAsset->AlignmentSections.Num();
-	
-	AlignmentSectionToScenePivotList.Reset(AlignmentSectionsNum);
-
-	for (int32 Idx = 0; Idx < AlignmentSectionsNum; Idx++)
+	for (const FContextualAnimAlignmentSectionData AligmentData : SceneAsset->GetAlignmentSections())
 	{
-		FTransform ScenePivot = FTransform::Identity;
-
-		if(SceneAsset->AlignmentSections[Idx].ScenePivotProvider)
-		{
-			ScenePivot = SceneAsset->AlignmentSections[Idx].ScenePivotProvider->CalculateScenePivot_Runtime(SceneActorMap);
-		}
-		else
+		FTransform ScenePivotRuntime = FTransform::Identity;
+		if (const FContextualAnimSceneActorData* SceneActor = FindSceneActorDataByRole(AligmentData.Origin))
 		{
-			ScenePivot = SceneActorMap.FindRef(SceneAsset->PrimaryRole).GetTransform();
-		}
+			if (AligmentData.bAlongClosestDistance)
+			{
+				if (const FContextualAnimSceneActorData* OtherSceneActor = FindSceneActorDataByRole(AligmentData.OtherRole))
+				{
+					const FTransform T1 = SceneActor->GetActor()->GetActorTransform();
+					const FTransform T2 = OtherSceneActor->GetActor()->GetActorTransform();
 
-		AlignmentSectionToScenePivotList.Add(MakeTuple(SceneAsset->AlignmentSections[Idx].SectionName, ScenePivot));
+					ScenePivotRuntime.SetLocation(FMath::Lerp<FVector>(T1.GetLocation(), T2.GetLocation(), AligmentData.Weight));
+					ScenePivotRuntime.SetRotation((T2.GetLocation() - T1.GetLocation()).GetSafeNormal2D().ToOrientationQuat());
+				}
+			}
+			else
+			{
+				ScenePivotRuntime = SceneActor->GetActor()->GetActorTransform();
+			}
+		}
 
-		//DrawDebugCoordinateSystem(GetWorld(), ScenePivot.GetLocation(), ScenePivot.Rotator(), 100.f, false, 10.f, 0, 5.f);
+		AlignmentSectionToScenePivotList.Add(MakeTuple(AligmentData.WarpTargetName, ScenePivotRuntime));
 	}
+	
 
-	for (auto& Pair : SceneActorMap)
+	for (auto& Binding : Bindings)
 	{
-		const FName& Role = Pair.Key;
-		FContextualAnimSceneActorData& Data = Pair.Value;
-
-		if (SceneAsset->GetTrackSettings(Role)->JoinRule == EContextualAnimJoinRule::Default)
-		{
-			Join(Data);
-		}
+		Join(Binding);
 	}
 }
 
 void UContextualAnimSceneInstance::Stop()
 {
-	for (auto& Pair : SceneActorMap)
+	for (auto& Binding : Bindings)
 	{
-		Leave(Pair.Value);
+		Leave(Binding);
 	}
 }
 
@@ -387,19 +276,18 @@ void UContextualAnimSceneInstance::OnMontageBlendingOut(UAnimMontage* Montage, b
 {
 	UE_LOG(LogContextualAnim, Log, TEXT("UContextualAnimSceneInstance::OnMontageBlendingOut Montage: %s"), *GetNameSafe(Montage));
 
-	for (auto& Pair : SceneActorMap)
+	for (auto& Binding : Bindings)
 	{
-		FContextualAnimSceneActorData& SceneActorData = Pair.Value;
-		if (SceneActorData.GetAnimData()->Animation == Montage)
+		if (Binding.GetAnimTrack().Animation == Montage)
 		{
-			AActor* Actor = SceneActorData.GetActor();
-			if (UAnimInstance* AnimInstance = SceneActorData.GetAnimInstance())
+			AActor* Actor = Binding.GetActor();
+			if (UAnimInstance* AnimInstance = Binding.GetAnimInstance())
 			{
 				AnimInstance->OnPlayMontageNotifyBegin.RemoveDynamic(this, &UContextualAnimSceneInstance::OnNotifyBeginReceived);
 				AnimInstance->OnPlayMontageNotifyEnd.RemoveDynamic(this, &UContextualAnimSceneInstance::OnNotifyEndReceived);
 				AnimInstance->OnMontageBlendingOut.RemoveDynamic(this, &UContextualAnimSceneInstance::OnMontageBlendingOut);
 
-				if (SceneActorData.GetAnimData()->bRequireFlyingMode)
+				if (Binding.GetAnimTrack().bRequireFlyingMode)
 				{
 					if (UCharacterMovementComponent* CharacterMovementComp = Actor->FindComponentByClass<UCharacterMovementComponent>())
 					{
@@ -408,29 +296,28 @@ void UContextualAnimSceneInstance::OnMontageBlendingOut(UAnimMontage* Montage, b
 				}
 			}
 
-			if (SceneAsset->bDisableCollisionBetweenActors)
+			if (SceneAsset->GetDisableCollisionBetweenActors())
 			{
-				SetIgnoreCollisionWithOtherActors(SceneActorData.GetActor(), false);
+				SetIgnoreCollisionWithOtherActors(Binding.GetActor(), false);
 			}
 
-			if (UContextualAnimSceneActorComponent* SceneActorComp = Pair.Value.GetSceneActorComponent())
+			if (UContextualAnimSceneActorComponent* SceneActorComp = Binding.GetSceneActorComponent())
 			{
-				SceneActorComp->OnLeftScene(&SceneActorData);
+				SceneActorComp->OnLeftScene(&Binding);
 			}
 
-			OnActorLeft.ExecuteIfBound(this, Actor);
+			OnActorLeft.Broadcast(this, Actor);
 
 			break;
 		}
 	}
 
 	bool bShouldEnd = true;
-	for (auto& Pair : SceneActorMap)
+	for (auto& Binding : Bindings)
 	{
-		const FContextualAnimSceneActorData& Data = Pair.Value;
-		if (UAnimInstance* AnimInstance = Data.GetAnimInstance())
+		if (UAnimInstance* AnimInstance = Binding.GetAnimInstance())
 		{
-			if(AnimInstance->Montage_IsPlaying(Data.GetAnimData()->Animation))
+			if(AnimInstance->Montage_IsPlaying(Binding.GetAnimTrack().Animation))
 			{
 				bShouldEnd = false;
 				break;
@@ -440,7 +327,7 @@ void UContextualAnimSceneInstance::OnMontageBlendingOut(UAnimMontage* Montage, b
 
 	if(bShouldEnd)
 	{
-		OnSceneEnded.ExecuteIfBound(this);
+		OnSceneEnded.Broadcast(this);
 	}
 }
 
@@ -449,20 +336,9 @@ void UContextualAnimSceneInstance::OnNotifyBeginReceived(FName NotifyName, const
 	UE_LOG(LogContextualAnim, Log, TEXT("UContextualAnimSceneInstance::OnNotifyBeginReceived NotifyName: %s Montage: %s"),
 	*NotifyName.ToString(), *GetNameSafe(BranchingPointNotifyPayload.SequenceAsset));
 
-	// @TODO: For now just use a hard-coded name to identify the event. We should change this in the future
-	static const FName LateJoinNotifyName = FName(TEXT("ContextualAnimLateJoin"));
-	if(NotifyName == LateJoinNotifyName)
+	if(const USkeletalMeshComponent* SkelMeshCom = BranchingPointNotifyPayload.SkelMeshComponent)
 	{
-		for (auto& Pair : SceneActorMap)
-		{
-			const FName& Role = Pair.Key;
-			FContextualAnimSceneActorData& Data = Pair.Value;
-
-			if (SceneAsset->GetTrackSettings(Role)->JoinRule == EContextualAnimJoinRule::Late)
-			{
-				Join(Data);
-			}
-		}
+		OnNotifyBegin.Broadcast(this, SkelMeshCom->GetOwner(), NotifyName);
 	}
 }
 
@@ -470,13 +346,18 @@ void UContextualAnimSceneInstance::OnNotifyEndReceived(FName NotifyName, const F
 {
 	UE_LOG(LogContextualAnim, Log, TEXT("UContextualAnimSceneInstance::OnNotifyEndReceived NotifyName: %s Montage: %s"),
 	*NotifyName.ToString(), *GetNameSafe(BranchingPointNotifyPayload.SequenceAsset));
+
+	if (const USkeletalMeshComponent* SkelMeshCom = BranchingPointNotifyPayload.SkelMeshComponent)
+	{
+		OnNotifyEnd.Broadcast(this, SkelMeshCom->GetOwner(), NotifyName);
+	}
 }
 
 void UContextualAnimSceneInstance::SetIgnoreCollisionWithOtherActors(AActor* Actor, bool bValue) const
 {
-	for (auto& Pair : SceneActorMap)
+	for (auto& Binding : Bindings)
 	{
-		AActor* OtherActor = Pair.Value.GetActor();
+		AActor* OtherActor = Binding.GetActor();
 		if(OtherActor != Actor)
 		{
 			if (UPrimitiveComponent* RootPrimitiveComponent = Cast<UPrimitiveComponent>(Actor->GetRootComponent()))
@@ -485,66 +366,4 @@ void UContextualAnimSceneInstance::SetIgnoreCollisionWithOtherActors(AActor* Act
 			}
 		}
 	}
-}
-
-float UContextualAnimSceneInstance::GetCurrentSectionTimeLeft() const
-{
-	if(SceneAsset)
-	{
-		if(const FContextualAnimSceneActorData* Data = FindSceneActorDataForRole(SceneAsset->GetLeaderRole()))
-		{
-			if (const FAnimMontageInstance* MontageInstance = Data->GetAnimMontageInstance())
-			{
-				return MontageInstance->Montage->GetSectionTimeLeftFromPos(MontageInstance->GetPosition());
-			}
-		}
-	}
-	
-	return 0.f;
-}
-
-float UContextualAnimSceneInstance::GetPositionInCurrentSection() const
-{
-	if (SceneAsset)
-	{
-		if (const FContextualAnimSceneActorData* Data = FindSceneActorDataForRole(SceneAsset->GetLeaderRole()))
-		{
-			if (const FAnimMontageInstance* MontageInstance = Data->GetAnimMontageInstance())
-			{
-				float PosInSection = 0.f;
-				MontageInstance->Montage->GetAnimCompositeSectionIndexFromPos(MontageInstance->GetPosition(), PosInSection);
-				return PosInSection;
-			}
-		}
-	}
-
-	return 0.f;
-}
-
-bool UContextualAnimSceneInstance::DidCurrentSectionLoop() const
-{
-	if (SceneAsset)
-	{
-		if (const FContextualAnimSceneActorData* Data = FindSceneActorDataForRole(SceneAsset->GetLeaderRole()))
-		{
-			if (const FAnimMontageInstance* MontageInstance = Data->GetAnimMontageInstance())
-			{
-				const float PreviousPos = MontageInstance->GetPreviousPosition();
-				const float CurrentPos = MontageInstance->GetPosition();
-
-				const int32 SectionIdxPreviousPos = MontageInstance->Montage->GetSectionIndexFromPosition(PreviousPos);
-				const int32 SectionIdxCurrentPos = MontageInstance->Montage->GetSectionIndexFromPosition(CurrentPos);
-
-				if(SectionIdxPreviousPos == SectionIdxCurrentPos)
-				{
-					const float TimeLeftFromPreviousPos = MontageInstance->Montage->GetSectionTimeLeftFromPos(PreviousPos);
-					const float TimeLeftFromCurrentPos = MontageInstance->Montage->GetSectionTimeLeftFromPos(CurrentPos);
-
-					return TimeLeftFromPreviousPos < TimeLeftFromCurrentPos;
-				}
-			}
-		}
-	}
-
-	return false;
-}
+}
\ No newline at end of file
diff --git a/Engine/Plugins/Experimental/Animation/ContextualAnimation/Source/ContextualAnimation/Private/ContextualAnimScenePivotProvider.cpp b/Engine/Plugins/Experimental/Animation/ContextualAnimation/Source/ContextualAnimation/Private/ContextualAnimScenePivotProvider.cpp
deleted file mode 100644
diff --git a/Engine/Plugins/Experimental/Animation/ContextualAnimation/Source/ContextualAnimation/Private/ContextualAnimSelectionCriterion.cpp b/Engine/Plugins/Experimental/Animation/ContextualAnimation/Source/ContextualAnimation/Private/ContextualAnimSelectionCriterion.cpp
new file mode 100644
--- /dev/null
+++ b/Engine/Plugins/Experimental/Animation/ContextualAnimation/Source/ContextualAnimation/Private/ContextualAnimSelectionCriterion.cpp
@@ -0,0 +1,75 @@
+// Copyright Epic Games, Inc. All Rights Reserved.
+
+#include "ContextualAnimSelectionCriterion.h"
+#include "ContextualAnimSceneAsset.h"
+#include "Engine/World.h"
+#include "DrawDebugHelpers.h"
+#include "UObject/ObjectSaveContext.h"
+#include "SceneManagement.h"
+
+UContextualAnimSelectionCriterion::UContextualAnimSelectionCriterion(const FObjectInitializer& ObjectInitializer)
+	: Super(ObjectInitializer)
+{
+}
+
+UContextualAnimSceneAsset* UContextualAnimSelectionCriterion::GetSceneAssetOwner() const
+{
+	return Cast<UContextualAnimSceneAsset>(GetOuter());
+}
+
+// UContextualAnimSelectionCriterion_TriggerArea
+//===========================================================================
+
+UContextualAnimSelectionCriterion_TriggerArea::UContextualAnimSelectionCriterion_TriggerArea(const FObjectInitializer& ObjectInitializer)
+	: Super(ObjectInitializer)
+{
+	//@TODO: I think we could initialize this from the animations so it automatically creates an area from the primary to the owner of this criterion
+	PolygonPoints.Add(FVector(100, -100, 0));
+	PolygonPoints.Add(FVector(-100, -100, 0));
+	PolygonPoints.Add(FVector(-100, 100, 0));
+	PolygonPoints.Add(FVector(100, 100, 0));
+}
+
+bool UContextualAnimSelectionCriterion_TriggerArea::DoesQuerierPassCondition(const FContextualAnimPrimaryActorData& PrimaryActorData, const FContextualAnimQuerierData& QuerierData) const
+{
+	if (PolygonPoints.Num() == 4)
+	{
+		const float HalfHeight = FMath::Max((Height / 2.f), 0.f);
+		const float VDist = FMath::Abs((PrimaryActorData.Transform.GetLocation().Z + PolygonPoints[0].Z + HalfHeight) - QuerierData.Transform.GetLocation().Z);
+		if (VDist > HalfHeight)
+		{
+			return false;
+		}
+
+		const FVector2D TestPoint = FVector2D(QuerierData.Transform.GetLocation());
+		const int32 NumPoints = PolygonPoints.Num();
+		float AngleSum = 0.0f;
+		for (int32 PointIndex = 0; PointIndex < NumPoints; ++PointIndex)
+		{
+			const FVector2D& VecAB = FVector2D(PrimaryActorData.Transform.TransformPositionNoScale(PolygonPoints[PointIndex])) - TestPoint;
+			const FVector2D& VecAC = FVector2D(PrimaryActorData.Transform.TransformPositionNoScale(PolygonPoints[(PointIndex + 1) % NumPoints])) - TestPoint;
+			const float Angle = FMath::Sign(FVector2D::CrossProduct(VecAB, VecAC)) * FMath::Acos(FMath::Clamp(FVector2D::DotProduct(VecAB, VecAC) / (VecAB.Size() * VecAC.Size()), -1.0f, 1.0f));
+			AngleSum += Angle;
+		}
+
+		return (FMath::Abs(AngleSum) > 0.001f);
+	}
+
+	return false;
+}
+
+// UContextualAnimSelectionCriterion_Facing
+//===========================================================================
+
+bool UContextualAnimSelectionCriterion_Facing::DoesQuerierPassCondition(const FContextualAnimPrimaryActorData& PrimaryActorData, const FContextualAnimQuerierData& QuerierData) const
+{
+	if (MaxAngle > 0.f)
+	{
+		const float FacingCos = FMath::Cos(FMath::Clamp(FMath::DegreesToRadians(MaxAngle), 0.f, PI));
+		const FVector ToTarget = (PrimaryActorData.Transform.GetLocation() - QuerierData.Transform.GetLocation()).GetSafeNormal2D();
+		const float DotProduct = FVector::DotProduct(QuerierData.Transform.GetRotation().GetForwardVector(), ToTarget);
+		return (DotProduct > FacingCos);
+	}
+
+	return true;
+}
\ No newline at end of file
diff --git a/Engine/Plugins/Experimental/Animation/ContextualAnimation/Source/ContextualAnimation/Private/ContextualAnimTypes.cpp b/Engine/Plugins/Experimental/Animation/ContextualAnimation/Source/ContextualAnimation/Private/ContextualAnimTypes.cpp
--- a/Engine/Plugins/Experimental/Animation/ContextualAnimation/Source/ContextualAnimation/Private/ContextualAnimTypes.cpp
+++ b/Engine/Plugins/Experimental/Animation/ContextualAnimation/Source/ContextualAnimation/Private/ContextualAnimTypes.cpp
@@ -7,11 +7,13 @@
 #include "AnimNotifyState_MotionWarping.h"
 #include "ContextualAnimUtilities.h"
 #include "RootMotionModifier.h"
+#include "ContextualAnimSelectionCriterion.h"
 
 DEFINE_LOG_CATEGORY(LogContextualAnim);
 
-const FContextualAnimData FContextualAnimData::EmptyAnimData;
+const FContextualAnimTrack FContextualAnimTrack::EmptyTrack;
 const FContextualAnimIKTarget FContextualAnimIKTarget::InvalidIKTarget;
+const FContextualAnimIKTargetDefContainer FContextualAnimIKTargetDefContainer::EmptyContainer;
 
 // FContextualAnimAlignmentTrackContainer
 ///////////////////////////////////////////////////////////////////////
@@ -37,7 +39,7 @@ FTransform FContextualAnimAlignmentTrackContainer::ExtractTransformAtTime(int32
 	return AlignmentTransform;
 }
 
-float FContextualAnimData::GetSyncTimeForWarpSection(int32 WarpSectionIndex) const
+float FContextualAnimTrack::GetSyncTimeForWarpSection(int32 WarpSectionIndex) const
 {
 	//@TODO: We need a better way to identify warping sections withing the animation. This is just a temp solution
 	//@TODO: We should cache this data
@@ -97,7 +99,7 @@ float FContextualAnimData::GetSyncTimeForWarpSection(int32 WarpSectionIndex) con
 	return Result;
 }
 
-float FContextualAnimData::GetSyncTimeForWarpSection(const FName& WarpSectionName) const
+float FContextualAnimTrack::GetSyncTimeForWarpSection(const FName& WarpSectionName) const
 {
 	//@TODO: We need a better way to identify warping sections within the animation. This is just a temp solution
 	//@TODO: We should cache this data
@@ -130,7 +132,7 @@ float FContextualAnimData::GetSyncTimeForWarpSection(const FName& WarpSectionNam
 	return Result;
 }
 
-float FContextualAnimData::FindBestAnimStartTime(const FVector& LocalLocation) const
+float FContextualAnimTrack::FindBestAnimStartTime(const FVector& LocalLocation) const
 {
 	float BestTime = 0.f;
 
@@ -169,20 +171,15 @@ float FContextualAnimData::FindBestAnimStartTime(const FVector& LocalLocation) c
 	return BestTime;
 }
 
-FTransform FContextualAnimCompositeTrack::GetRootTransformForAnimDataAtIndex(int32 Index) const
+bool FContextualAnimTrack::DoesQuerierPassSelectionCriteria(const FContextualAnimPrimaryActorData& PrimaryActorData, const FContextualAnimQuerierData& QuerierData) const
 {
-	if(AnimDataContainer.IsValidIndex(Index))
+	for (const UContextualAnimSelectionCriterion* Criterion : SelectionCriteria)
 	{
-		const FContextualAnimData& AnimData = AnimDataContainer[Index];
-		if (AnimData.Animation)
+		if (Criterion && !Criterion->DoesQuerierPassCondition(PrimaryActorData, QuerierData))
 		{
-			return (Settings.MeshToComponent.Inverse() * (UContextualAnimUtilities::ExtractRootTransformFromAnimation(AnimData.Animation, 0.f) * AnimData.MeshToScene));
-		}
-		else
-		{
-			return Settings.MeshToComponent.Inverse() * AnimData.MeshToScene;
+			return false;
 		}
 	}
 
-	return FTransform::Identity;
+	return true;
 }
\ No newline at end of file
diff --git a/Engine/Plugins/Experimental/Animation/ContextualAnimation/Source/ContextualAnimation/Private/ContextualAnimUtilities.cpp b/Engine/Plugins/Experimental/Animation/ContextualAnimation/Source/ContextualAnimation/Private/ContextualAnimUtilities.cpp
--- a/Engine/Plugins/Experimental/Animation/ContextualAnimation/Source/ContextualAnimation/Private/ContextualAnimUtilities.cpp
+++ b/Engine/Plugins/Experimental/Animation/ContextualAnimation/Source/ContextualAnimation/Private/ContextualAnimUtilities.cpp
@@ -18,6 +18,7 @@
 #include "Components/SkeletalMeshComponent.h"
 #include "EngineUtils.h"
 #include "AnimNotifyState_IKWindow.h"
+#include "SceneManagement.h"
 
 void UContextualAnimUtilities::ExtractLocalSpacePose(const UAnimSequenceBase* Animation, const FBoneContainer& BoneContainer, float Time, bool bExtractRootMotion, FCompactPose& OutPose)
 {
@@ -88,6 +89,8 @@ void UContextualAnimUtilities::DrawDebugPose(const UWorld* World, const UAnimSeq
 { 
 	if(World)
 	{
+		FMemMark Mark(FMemStack::Get());
+
 		Time = FMath::Clamp(Time, 0.f, Animation->GetPlayLength());
 
 		const int32 TotalBones = Animation->GetSkeleton()->GetReferenceSkeleton().GetNum();
@@ -127,26 +130,25 @@ void UContextualAnimUtilities::DrawDebugPose(const UWorld* World, const UAnimSeq
 	}
 }
 
-void UContextualAnimUtilities::DrawDebugScene(const UWorld* World, const UContextualAnimSceneAsset* SceneAsset, int32 AnimDataIndex, float Time, const FTransform& ToWorldTransform, const FColor& Color, float LifeTime, float Thickness)
+void UContextualAnimUtilities::DrawDebugScene(const UWorld* World, const UContextualAnimSceneAsset* SceneAsset, int32 VariantIdx, float Time, const FTransform& ToWorldTransform, const FColor& Color, float LifeTime, float Thickness)
 {
 	if (World && SceneAsset)
 	{
-		for(const auto& Pair : SceneAsset->DataContainer)
+		SceneAsset->ForEachAnimTrack(VariantIdx, [=](const FContextualAnimTrack& AnimTrack)
 		{
-			if (Pair.Value.AnimDataContainer.IsValidIndex(AnimDataIndex))
+			const FTransform Transform = (SceneAsset->GetMeshToComponentForRole(AnimTrack.Role) * AnimTrack.GetAlignmentTransformAtTime(Time)) * ToWorldTransform;
+			
+			if (const UAnimMontage* Animation = AnimTrack.Animation)
 			{
-				const FContextualAnimData& AnimData = Pair.Value.AnimDataContainer[AnimDataIndex];
-				const FTransform Transform = (Pair.Value.Settings.MeshToComponent * AnimData.GetAlignmentTransformAtTime(Time)) * ToWorldTransform;
-				if (const UAnimMontage* Animation = AnimData.Animation)
-				{
-					DrawDebugPose(World, Animation, Time, Transform, Color, LifeTime, Thickness);
-				}
-				else
-				{
-					DrawDebugCoordinateSystem(World, Transform.GetLocation(), Transform.Rotator(), 50.f, false, LifeTime, 0, Thickness);
-				}
+				DrawDebugPose(World, Animation, Time, Transform, Color, LifeTime, Thickness);
 			}
-		}
+			else
+			{
+				DrawDebugCoordinateSystem(World, Transform.GetLocation(), Transform.Rotator(), 50.f, false, LifeTime, 0, Thickness);
+			}
+
+			return UE::ContextualAnim::EForEachResult::Continue;
+		});
 	}
 }
 
@@ -209,4 +211,49 @@ float UContextualAnimUtilities::BP_Montage_GetSectionTimeLeftFromPos(const UAnim
 float UContextualAnimUtilities::BP_Montage_GetSectionLength(const UAnimMontage* Montage, int32 SectionIndex)
 {
 	return Montage ? Montage->GetSectionLength(SectionIndex) : -1.f;
+}
+
+void UContextualAnimUtilities::DrawSector(FPrimitiveDrawInterface& PDI, const FVector& Origin, const FVector& Direction, float MinDistance, float MaxDistance, float MinAngle, float MaxAngle, const FLinearColor& Color, uint8 DepthPriority, float Thickness)
+{
+	if(MinAngle == 0 && MaxAngle == 0)
+	{
+		DrawCircle(&PDI, Origin, FVector(1, 0, 0), FVector(0, 1, 0), Color, 30.f, 12, SDPG_World, 1.f);
+		return;
+	}
+
+	// Draw Cone lines
+	const FVector LeftDirection = Direction.RotateAngleAxis(MinAngle, FVector::UpVector);
+	const FVector RightDirection = Direction.RotateAngleAxis(MaxAngle, FVector::UpVector);
+	PDI.DrawLine(Origin + (LeftDirection * MinDistance), Origin + (LeftDirection * MaxDistance), Color, DepthPriority, Thickness);
+	PDI.DrawLine(Origin + (RightDirection * MinDistance), Origin + (RightDirection * MaxDistance), Color, DepthPriority, Thickness);
+
+	// Draw Near Arc
+	FVector LastDirection = LeftDirection;
+	float Angle = MinAngle;
+	while (Angle < MaxAngle)
+	{
+		Angle = FMath::Clamp<float>(Angle + 10, MinAngle, MaxAngle);
+
+		const float Length = MinDistance;
+		const FVector NewDirection = Direction.RotateAngleAxis(Angle, FVector::UpVector);
+		const FVector LineStart = Origin + (LastDirection * Length);
+		const FVector LineEnd = Origin + (NewDirection * Length);
+		PDI.DrawLine(LineStart, LineEnd, Color, DepthPriority, Thickness);
+		LastDirection = NewDirection;
+	}
+
+	// Draw Far Arc
+	LastDirection = LeftDirection;
+	Angle = MinAngle;
+	while (Angle < MaxAngle)
+	{
+		Angle = FMath::Clamp<float>(Angle + 10, MinAngle, MaxAngle);
+
+		const float Length = MaxDistance;
+		const FVector NewDirection = Direction.RotateAngleAxis(Angle, FVector::UpVector);
+		const FVector LineStart = Origin + (LastDirection * Length);
+		const FVector LineEnd = Origin + (NewDirection * Length);
+		PDI.DrawLine(LineStart, LineEnd, Color, DepthPriority, Thickness);
+		LastDirection = NewDirection;
+	}
 }
\ No newline at end of file
diff --git a/Engine/Plugins/Experimental/Animation/ContextualAnimation/Source/ContextualAnimation/Public/AnimNotifyState_IKWindow.h b/Engine/Plugins/Experimental/Animation/ContextualAnimation/Source/ContextualAnimation/Public/AnimNotifyState_IKWindow.h
--- a/Engine/Plugins/Experimental/Animation/ContextualAnimation/Source/ContextualAnimation/Public/AnimNotifyState_IKWindow.h
+++ b/Engine/Plugins/Experimental/Animation/ContextualAnimation/Source/ContextualAnimation/Public/AnimNotifyState_IKWindow.h
@@ -30,4 +30,6 @@ public:
 	UAnimNotifyState_IKWindow(const FObjectInitializer& ObjectInitializer);
 
 	virtual FString GetNotifyName_Implementation() const override;
+
+	static float GetIKAlphaValue(const FName& GoalName, const struct FAnimMontageInstance* MontageInstance);
 };
\ No newline at end of file
diff --git a/Engine/Plugins/Experimental/Animation/ContextualAnimation/Source/ContextualAnimation/Public/ContextualAnimManager.h b/Engine/Plugins/Experimental/Animation/ContextualAnimation/Source/ContextualAnimation/Public/ContextualAnimManager.h
--- a/Engine/Plugins/Experimental/Animation/ContextualAnimation/Source/ContextualAnimation/Public/ContextualAnimManager.h
+++ b/Engine/Plugins/Experimental/Animation/ContextualAnimation/Source/ContextualAnimation/Public/ContextualAnimManager.h
@@ -21,6 +21,8 @@ class CONTEXTUALANIMATION_API UContextualAnimManager : public UObject, public FT
 
 public:
 
+	friend class FContextualAnimViewModel;
+
 	UContextualAnimManager(const FObjectInitializer& ObjectInitializer);
 
 	virtual UWorld* GetWorld() const override;
@@ -39,15 +41,12 @@ public:
 
 	void UnregisterSceneActorComponent(UContextualAnimSceneActorComponent* SceneActorComp);
 
-	/** Attempt to start an scene instance with the supplied bindings for each role */
-	bool TryStartScene(const UContextualAnimSceneAsset* SceneAsset, const FContextualAnimSceneBindings& Bindings);
-
-	/** Attempt to start an scene instance with PrimaryActor bound to the primary role and the first component valid for each of the other roles */
-	bool TryStartScene(const UContextualAnimSceneAsset* SceneAsset, AActor* PrimaryActor, const TSet<UContextualAnimSceneActorComponent*>& SceneActorComps);
-
-	UFUNCTION(BlueprintCallable, Category = "Contextual Anim|Manager")
-	UContextualAnimSceneActorComponent* FindClosestSceneActorCompToActor(const AActor* Actor) const;
+	/** Starts an scene instance with the supplied actors for each role ignoring selection criteria */
+	UContextualAnimSceneInstance* ForceStartScene(const UContextualAnimSceneAsset& SceneAsset, const FContextualAnimStartSceneParams& Params);
 	
+	/** Attempts to start an scene instance with the supplied actors using the first valid variant based on selection criteria */
+	UContextualAnimSceneInstance* TryStartScene(const UContextualAnimSceneAsset& SceneAsset, const FContextualAnimStartSceneParams& Params);
+
 	UFUNCTION(BlueprintCallable, Category = "Contextual Anim|Manager")
 	bool TryStopSceneWithActor(AActor* Actor);
 
@@ -62,8 +61,9 @@ public:
 
 	FORCEINLINE const TSet<UContextualAnimSceneActorComponent*>& GetSceneActorCompContainer() const { return SceneActorCompContainer; };
 
-	UFUNCTION(BlueprintCallable, Category = "Contextual Anim|Manager", meta = (DisplayName = "TryStartSceneWithBindings"))
-	bool BP_TryStartSceneWithBindings(const UContextualAnimSceneAsset* SceneAsset, const FContextualAnimSceneBindings& Bindings) { return TryStartScene(SceneAsset, Bindings); }
+	/** Attempts to start an scene instance with the supplied actors using the first valid variant based on selection criteria */
+	UFUNCTION(BlueprintCallable, Category = "Contextual Anim|Manager", meta = (DisplayName = "Try Start Scene"))
+	UContextualAnimSceneInstance* BP_TryStartScene(const UContextualAnimSceneAsset* SceneAsset, const FContextualAnimStartSceneParams& Params);
 
 protected:
 
@@ -74,5 +74,6 @@ protected:
 	UPROPERTY()
 	TArray<UContextualAnimSceneInstance*> Instances;
 
+	UFUNCTION()
 	void OnSceneInstanceEnded(UContextualAnimSceneInstance* SceneInstance);
 };
\ No newline at end of file
diff --git a/Engine/Plugins/Experimental/Animation/ContextualAnimation/Source/ContextualAnimation/Public/ContextualAnimMetadata.h b/Engine/Plugins/Experimental/Animation/ContextualAnimation/Source/ContextualAnimation/Public/ContextualAnimMetadata.h
deleted file mode 100644
diff --git a/Engine/Plugins/Experimental/Animation/ContextualAnimation/Source/ContextualAnimation/Public/ContextualAnimSceneActorComponent.h b/Engine/Plugins/Experimental/Animation/ContextualAnimation/Source/ContextualAnimation/Public/ContextualAnimSceneActorComponent.h
--- a/Engine/Plugins/Experimental/Animation/ContextualAnimation/Source/ContextualAnimation/Public/ContextualAnimSceneActorComponent.h
+++ b/Engine/Plugins/Experimental/Animation/ContextualAnimation/Source/ContextualAnimation/Public/ContextualAnimSceneActorComponent.h
@@ -3,8 +3,8 @@
 #pragma once
 
 #include "CoreMinimal.h"
-#include "Components/SphereComponent.h"
 #include "ContextualAnimTypes.h"
+#include "ActorComponents/IKRigInterface.h"
 #include "ContextualAnimSceneActorComponent.generated.h"
 
 class UAnimInstance;
@@ -13,22 +13,10 @@ class AActor;
 class UContextualAnimSceneInstance;
 struct FContextualAnimSceneActorData;
 
-USTRUCT(BlueprintType)
-struct FContextualAnimDebugParams
-{
-	GENERATED_BODY()
-
-	UPROPERTY(EditAnywhere, BlueprintReadWrite, Category = "Debug")
-	TWeakObjectPtr<AActor> TestActor;
-
-	UPROPERTY(EditAnywhere, BlueprintReadWrite, Category = "Debug", meta = (UIMin = 0, ClampMin = 0))
-	float DrawAlignmentTransformAtTime = 0.f;
-};
-
 DECLARE_DYNAMIC_MULTICAST_DELEGATE_OneParam(FContextualAnimSceneActorCompDelegate, class UContextualAnimSceneActorComponent*, SceneActorComponent);
 
 UCLASS(meta = (BlueprintSpawnableComponent))
-class CONTEXTUALANIMATION_API UContextualAnimSceneActorComponent : public UPrimitiveComponent
+class CONTEXTUALANIMATION_API UContextualAnimSceneActorComponent : public UPrimitiveComponent, public IIKGoalCreatorInterface
 {
 	GENERATED_BODY()
 
@@ -48,9 +36,6 @@ public:
 	UPROPERTY(EditAnywhere, BlueprintReadWrite, Category = "Debug")
 	bool bEnableDebug;
 
-	UPROPERTY(EditAnywhere, BlueprintReadWrite, Category = "Debug", meta = (ShowOnlyInnerProperties, EditCondition = "bEnableDebug"))
-	FContextualAnimDebugParams DebugParams;
-
 	UContextualAnimSceneActorComponent(const FObjectInitializer& ObjectInitializer);
 
 	virtual void OnRegister() override;
@@ -58,6 +43,8 @@ public:
 	virtual FBoxSphereBounds CalcBounds(const FTransform& LocalToWorld) const;
 	virtual FPrimitiveSceneProxy* CreateSceneProxy() override;
 
+	virtual void AddIKGoals_Implementation(TMap<FName, FIKRigGoal>& OutGoals) override;
+
 	/** Called from the scene instance when the actor owner of this component joins an scene */
 	void OnJoinedScene(const FContextualAnimSceneActorData* SceneActorData);
 	
diff --git a/Engine/Plugins/Experimental/Animation/ContextualAnimation/Source/ContextualAnimation/Public/ContextualAnimSceneAsset.h b/Engine/Plugins/Experimental/Animation/ContextualAnimation/Source/ContextualAnimation/Public/ContextualAnimSceneAsset.h
--- a/Engine/Plugins/Experimental/Animation/ContextualAnimation/Source/ContextualAnimation/Public/ContextualAnimSceneAsset.h
+++ b/Engine/Plugins/Experimental/Animation/ContextualAnimation/Source/ContextualAnimation/Public/ContextualAnimSceneAsset.h
@@ -15,111 +15,170 @@ struct FContextualAnimAlignmentSectionData
 {
 	GENERATED_BODY()
 
-	UPROPERTY(EditAnywhere, BlueprintReadWrite, Category = "Alignment")
-	FName SectionName = NAME_None;
+	UPROPERTY(EditAnywhere, BlueprintReadWrite, Category = "Defaults")
+	FName WarpTargetName = NAME_None;
 
-	UPROPERTY(EditAnywhere, Instanced, Category = "Alignment")
-	UContextualAnimScenePivotProvider* ScenePivotProvider = nullptr;
+	UPROPERTY(EditAnywhere, BlueprintReadWrite, Category = "Defaults", meta = (GetOptions = "GetRoles"))
+	FName Origin = NAME_None;
+		
+	UPROPERTY(EditAnywhere, BlueprintReadWrite, Category = "Defaults")
+	bool bAlongClosestDistance = false;
 
-	UPROPERTY(VisibleAnywhere, BlueprintReadWrite, Category = "Alignment", AdvancedDisplay)
-	TArray<FTransform> ScenePivots;
+	UPROPERTY(EditAnywhere, BlueprintReadWrite, Category = "Defaults", meta = (GetOptions = "GetRoles", EditCondition = "bAlongClosestDistance"))
+	FName OtherRole = NAME_None;
+	
+	UPROPERTY(EditAnywhere, BlueprintReadWrite, Category = "Defaults", meta = (ClampMin = "0.0", ClampMax = "1.0", UIMin = "0.0", UIMax = "1.0", EditCondition = "bAlongClosestDistance"))
+	float Weight = 0.f;
 };
 
-enum class EContextualAnimForEachResult : uint8
+USTRUCT(BlueprintType)
+struct FContextualAnimRoleDefinition
 {
-	Break,
-	Continue,
+	GENERATED_BODY()
+
+	UPROPERTY(EditAnywhere, BlueprintReadWrite, Category = "Defaults")
+	FName Name;
+
+	UPROPERTY(EditAnywhere, BlueprintReadWrite, Category = "Defaults")
+	TSubclassOf<AActor> PreviewActorClass;
+
+	UPROPERTY(EditAnywhere, BlueprintReadWrite, Category = "Defaults")
+	FTransform MeshToComponent = FTransform(FRotator(0.f, -90.f, 0.f));
 };
 
 UCLASS(Blueprintable)
-class CONTEXTUALANIMATION_API UContextualAnimSceneAsset : public UDataAsset
+class CONTEXTUALANIMATION_API UContextualAnimRolesAsset : public UDataAsset
 {
 	GENERATED_BODY()
 
 public:
 
-	typedef TFunction<EContextualAnimForEachResult(const FName& Role, const FContextualAnimData& AnimData)> FForEachAnimDataFunction;
+	UPROPERTY(EditAnywhere, BlueprintReadWrite, Category = "Defaults")
+	TArray<FContextualAnimRoleDefinition> Roles;
 
-	UPROPERTY(EditAnywhere, BlueprintReadWrite, Category = "Animations")
-	TMap<FName, FContextualAnimCompositeTrack> DataContainer;
+	UContextualAnimRolesAsset(const FObjectInitializer& ObjectInitializer) : Super(ObjectInitializer) {};
 
-	UPROPERTY(EditAnywhere, Category = "Settings")
-	TSubclassOf<UContextualAnimSceneInstance> SceneInstanceClass;
+	const FContextualAnimRoleDefinition* FindRoleDefinitionByName(const FName& Name) const
+	{
+		return Roles.FindByPredicate([Name](const FContextualAnimRoleDefinition& RoleDef) { return RoleDef.Name == Name; });
+	}
+};
 
-	UPROPERTY(EditAnywhere, Category = "Settings")
-	bool bDisableCollisionBetweenActors;
+USTRUCT(BlueprintType)
+struct FContextualAnimTracksContainer
+{
+	GENERATED_BODY()
 
-	UPROPERTY(EditAnywhere, Category = "Alignment")
-	TArray<FContextualAnimAlignmentSectionData> AlignmentSections;
+	UPROPERTY(EditAnywhere, BlueprintReadWrite, Category = "Defaults")
+	TArray<FContextualAnimTrack> Tracks;
+
+	UPROPERTY(EditAnywhere, BlueprintReadWrite, Category = "Defaults")
+	TArray<FTransform> ScenePivots;
+};
 
-	UPROPERTY(EditAnywhere, Category = "Alignment", meta = (ClampMin = "1", ClampMax = "60"), AdvancedDisplay)
-	int32 SampleRate;
+UCLASS(Blueprintable)
+class CONTEXTUALANIMATION_API UContextualAnimSceneAsset : public UDataAsset
+{
+	GENERATED_BODY()
 
-	UPROPERTY(EditAnywhere, Category = "Transitions")
-	TArray<FContextualAnimTransitionContainer> Transitions;
+public:
 
-	//@TODO: Rename this is used for selection and move to private
-	UPROPERTY(EditAnywhere, Category = "Roles")
-	FName PrimaryRole;
+	typedef TFunctionRef<UE::ContextualAnim::EForEachResult(const FContextualAnimTrack& AnimTrack)> FForEachAnimTrackFunction;
 
 	UContextualAnimSceneAsset(const FObjectInitializer& ObjectInitializer);
 
+#if WITH_EDITOR
+	virtual void PostEditChangeProperty(FPropertyChangedEvent& PropertyChangedEvent) override;
+#endif
+
 	virtual void PreSave(FObjectPreSaveContext ObjectSaveContext) override;
 
 	void PrecomputeData();
-
-	const FContextualAnimTrackSettings* GetTrackSettings(const FName& Role) const;
 	
-	const FContextualAnimData* GetAnimDataForRoleAtIndex(const FName& Role, int32 Index) const;
-	
-	void ForEachAnimData(FForEachAnimDataFunction Function) const;
+	void ForEachAnimTrack(FForEachAnimTrackFunction Function) const;
 
-	TArray<FName> GetRoles() const;
+	void ForEachAnimTrack(int32 VariantIdx, FForEachAnimTrackFunction Function) const;
 
-	FORCEINLINE FName GetAlignmentSectionNameAtIndex(int32 Index) const
-	{
-		return AlignmentSections.IsValidIndex(Index) ? AlignmentSections[Index].SectionName : NAME_None;
-	}
+	FORCEINLINE const FName& GetPrimaryRole() const { return PrimaryRole; }
+	FORCEINLINE float GetRadius() const { return Radius; }
+	FORCEINLINE bool GetDisableCollisionBetweenActors() const { return bDisableCollisionBetweenActors; }
+	FORCEINLINE const TSubclassOf<UContextualAnimSceneInstance>& GetSceneInstanceClass() const { return SceneInstanceClass; }
+	FORCEINLINE const TArray<FContextualAnimAlignmentSectionData>& GetAlignmentSections() const { return  AlignmentSections; }
 
-	FORCEINLINE const FName& GetLeaderRole() const { return (LeaderRole != NAME_None) ? LeaderRole : PrimaryRole; }
+	UFUNCTION()
+	TArray<FName> GetRoles() const;
 
-	FORCEINLINE float GetRadius() const { return Radius; }
+	const FContextualAnimTrack* GetAnimTrack(const FName& Role, int32 VariantIdx) const;
 
 	UFUNCTION(BlueprintCallable, Category = "Contextual Anim|Scene Asset")
-	bool Query(const FName& Role, FContextualAnimQueryResult& OutResult, const FContextualAnimQueryParams& QueryParams, const FTransform& ToWorldTransform) const;
+	int32 FindVariantIdx(FName Role, UAnimMontage* Animation) const;
 
-	UFUNCTION(BlueprintCallable, Category = "Contextual Anim|Scene Asset")
-	UAnimMontage* GetAnimationForRoleAtIndex(FName Role, int32 Index) const;
+	FName FindRoleByAnimation(const UAnimMontage* Animation) const;
+
+	const FContextualAnimTrack* FindFirstAnimTrackForRoleThatPassesSelectionCriteria(const FName& Role, const FContextualAnimPrimaryActorData& PrimaryActorData, const FContextualAnimQuerierData& QuerierData) const;
+	
+	const FContextualAnimTrack* FindAnimTrackForRoleWithClosestEntryLocation(const FName& Role, const FContextualAnimPrimaryActorData& PrimaryActorData, const FVector& TestLocation) const;
 
 	UFUNCTION(BlueprintCallable, Category = "Contextual Anim|Scene Asset")
-	int32 FindAnimIndex(FName Role, UAnimMontage* Animation) const;
+	FTransform GetAlignmentTransformForRoleRelativeToScenePivot(FName Role, int32 VariantIdx, float Time) const;
 
 	UFUNCTION(BlueprintCallable, Category = "Contextual Anim|Scene Asset")
-	FTransform ExtractAlignmentTransformAtTime(FName Role, int32 AnimDataIndex, float Time) const;
+	FTransform GetAlignmentTransformForRoleRelativeToOtherRole(FName FromRole, FName ToRole, int32 VariantIdx, float Time) const;
 
 	UFUNCTION(BlueprintCallable, Category = "Contextual Anim|Scene Asset")
-	FTransform ExtractIKTargetTransformAtTime(FName Role, int32 AnimDataIndex, FName TrackName, float Time) const;
+	FTransform GetIKTargetTransformForRoleAtTime(FName Role, int32 VariantIdx, FName TrackName, float Time) const;
 
-	bool QueryCompositeTrack(const FContextualAnimCompositeTrack* Track, FContextualAnimQueryResult& OutResult, const FContextualAnimQueryParams& QueryParams, const FTransform& ToWorldTransform) const;
+	const FContextualAnimIKTargetDefContainer& GetIKTargetDefsForRole(const FName& Role) const;
 
-	FName FindRoleByAnimation(const UAnimMontage* Animation) const;
+	const FTransform& GetMeshToComponentForRole(const FName& Role) const;
+
+	int32 GetTotalVariants() const;
 
-	TArray<FName> GetIKGoalsForRole(const FName& Role) const;
+	UFUNCTION(BlueprintCallable, Category = "Contextual Anim|Scene Asset", meta = (DisplayName = "Get Anim Track"))
+	const FContextualAnimTrack& BP_GetAnimTrack(const FName& Role, int32 VariantIdx) const
+	{
+		const FContextualAnimTrack* AnimTrack = GetAnimTrack(Role, VariantIdx);
+		return AnimTrack ? *AnimTrack : FContextualAnimTrack::EmptyTrack;
+	}
 
 protected:
 
-	void GenerateAlignmentTracks(const FContextualAnimTrackSettings& Settings, FContextualAnimData& AnimData);
+	UPROPERTY(EditAnywhere, Category = "Defaults")
+	TObjectPtr<UContextualAnimRolesAsset> RolesAsset;
 
-	void GenerateIKTargetTracks(const FContextualAnimTrackSettings& Settings, FContextualAnimData& AnimData);
+	UPROPERTY(EditAnywhere, Category = "Defaults", meta = (GetOptions = "GetRoles"))
+	FName PrimaryRole = NAME_None;
 
-	void UpdateRadius();
+	UPROPERTY(EditAnywhere, Category = "Defaults")
+	TArray<FContextualAnimTracksContainer> Variants;
+
+	UPROPERTY(EditAnywhere, Category = "Defaults")
+	TMap<FName, FContextualAnimIKTargetDefContainer> RoleToIKTargetDefsMap;
+
+	UPROPERTY(EditAnywhere, Category = "Defaults")
+	TArray<FContextualAnimAlignmentSectionData> AlignmentSections;
+
+	UPROPERTY(EditAnywhere, Category = "Settings")
+	TSubclassOf<UContextualAnimSceneInstance> SceneInstanceClass;
 
-private:
+	UPROPERTY(EditAnywhere, Category = "Settings")
+	bool bDisableCollisionBetweenActors = true;
 
-	UPROPERTY(EditAnywhere, Category = "Roles")
-	FName LeaderRole;
+	/** Sample rate (frames per second) used when sampling the animations to generate alignment and IK tracks */
+	UPROPERTY(EditAnywhere, Category = "Settings", meta = (ClampMin = "1", ClampMax = "60"), AdvancedDisplay)
+	int32 SampleRate = 15;
 
-	/** Radius that enclose all the entry points */
 	UPROPERTY(VisibleAnywhere, Category = "Settings")
 	float Radius = 0.f;
+
+	void GenerateAlignmentTracks();
+
+	void GenerateIKTargetTracks();
+
+	void UpdateRadius();
+
+	friend class UContextualAnimUtilities;
+	friend class FContextualAnimViewModel;
+	friend class FContextualAnimEdMode;
+	friend class FContextualAnimMovieSceneNotifyTrackEditor;
 };
diff --git a/Engine/Plugins/Experimental/Animation/ContextualAnimation/Source/ContextualAnimation/Public/ContextualAnimSceneInstance.h b/Engine/Plugins/Experimental/Animation/ContextualAnimation/Source/ContextualAnimation/Public/ContextualAnimSceneInstance.h
--- a/Engine/Plugins/Experimental/Animation/ContextualAnimation/Source/ContextualAnimation/Public/ContextualAnimSceneInstance.h
+++ b/Engine/Plugins/Experimental/Animation/ContextualAnimation/Source/ContextualAnimation/Public/ContextualAnimSceneInstance.h
@@ -10,8 +10,7 @@
 #include "ContextualAnimSceneInstance.generated.h"
 
 struct FAnimMontageInstance;
-struct FContextualAnimData;
-struct FContextualAnimTrackSettings;
+struct FContextualAnimTrack;
 class UContextualAnimSceneInstance;
 class UContextualAnimSceneActorComponent;
 class UContextualAnimSceneAsset;
@@ -26,24 +25,20 @@ struct CONTEXTUALANIMATION_API FContextualAnimSceneActorData
 {
 	GENERATED_BODY()
 
-		FContextualAnimSceneActorData()
-		: Actor(nullptr), AnimDataPtr(nullptr), SettingsPtr(nullptr), AnimStartTime(0.f)
-	{}
-
-	FContextualAnimSceneActorData(AActor* InActor, const FContextualAnimData* InAnimData, const FContextualAnimTrackSettings* InSettings, float InAnimStartTime = 0.f)
-		: Actor(InActor), AnimDataPtr(InAnimData), SettingsPtr(InSettings), AnimStartTime(InAnimStartTime)
+	FContextualAnimSceneActorData(){}
+	FContextualAnimSceneActorData(const FName& InRole, int32 InVariantIdx, AActor& InActor, const FContextualAnimTrack& InAnimTrack, float InAnimStartTime = 0.f)
+		: Role(InRole), VariantIdx(InVariantIdx), Actor(&InActor), AnimTrackPtr(&InAnimTrack), AnimStartTime(InAnimStartTime)
 	{}
 
 	/** Return a pointer to the actual actor in the world */
 	FORCEINLINE AActor* GetActor() const { return Actor.Get(); }
-
+	FORCEINLINE FName GetRole() const { return Role; }
+	FORCEINLINE int32 GetVariantIdx() const { return VariantIdx; }
 	FORCEINLINE float GetAnimStartTime() const { return AnimStartTime; }
+	FORCEINLINE const FContextualAnimTrack& GetAnimTrack() const { return *AnimTrackPtr; }
+	FORCEINLINE const UContextualAnimSceneInstance& GetSceneInstance() const { return *SceneInstancePtr.Get(); }
 
-	FORCEINLINE const FContextualAnimData* GetAnimData() const { return AnimDataPtr; }
-
-	FORCEINLINE const FContextualAnimTrackSettings* GetSettings() const { return SettingsPtr; }
-
-	FORCEINLINE const UContextualAnimSceneInstance* GetSceneInstance() const { return SceneInstancePtr.Get(); }
+	const FContextualAnimIKTargetDefContainer& GetIKTargetDefs() const;
 
 	/** Return the transform used for alignment for this scene actor */
 	FTransform GetTransform() const;
@@ -55,9 +50,11 @@ struct CONTEXTUALANIMATION_API FContextualAnimSceneActorData
 
 	int32 GetCurrentSectionIndex() const;
 
+	/** Returns the ActiveMontageInstance or null in the case of static actors */
 	FAnimMontageInstance* GetAnimMontageInstance() const;
 
-	const UAnimMontage* GetAnimMontage() const;
+	/** Returns the Animation this actor is playing or null in the case of static actors */
+	UAnimMontage* GetAnimMontage() const;
 
 	UAnimInstance* GetAnimInstance() const;
 
@@ -69,29 +66,79 @@ private:
 
 	friend UContextualAnimSceneInstance;
 
+	/** Role this actor is representing */
+	FName Role = NAME_None;
+
+	int32 VariantIdx = INDEX_NONE;
+
 	/** The actual actor in the world */
 	TWeakObjectPtr<AActor> Actor = nullptr;
 
 	/** Ptr to the animation data in the scene asset used by this actor */
-	const FContextualAnimData* AnimDataPtr = nullptr;
-
-	const FContextualAnimTrackSettings* SettingsPtr = nullptr;
+	const FContextualAnimTrack* AnimTrackPtr = nullptr;
 
 	/** Desired time to start the animation */
 	float AnimStartTime = 0.f;
 
 	/** Ptr back to the scene instance we belong to */
 	TWeakObjectPtr<const UContextualAnimSceneInstance> SceneInstancePtr = nullptr;
+
+#if WITH_EDITOR
+public:
+	/** Guid only used in editor to bind this actor to sequencer */
+	FGuid Guid;
+#endif
+
+};
+
+USTRUCT(BlueprintType)
+struct CONTEXTUALANIMATION_API FContextualAnimSceneBindings
+{
+	GENERATED_BODY()
+	
+	const FContextualAnimSceneActorData* FindSceneActorDataByActor(const AActor* Actor) const
+	{
+		return Actor ? Data.FindByPredicate([Actor](const auto& Item) { return Item.GetActor() == Actor; }) : nullptr;
+	}
+
+	const FContextualAnimSceneActorData* FindSceneActorDataByRole(const FName& Role) const
+	{
+		return Role != NAME_None ? Data.FindByPredicate([&Role](const auto& Item) { return Item.GetRole() == Role; }) : nullptr;
+	}
+
+#if WITH_EDITOR
+	const FContextualAnimSceneActorData* FindSceneActorDataByGuid(const FGuid& Guid) const
+	{
+		return Guid.IsValid() ? Data.FindByPredicate([&Guid](const auto& Item) { return Item.Guid == Guid; }) : nullptr;
+	}
+#endif
+
+	FORCEINLINE int32 Num() const { return Data.Num(); }
+	FORCEINLINE int32 Add(const FContextualAnimSceneActorData& NewData) { return Data.Add(NewData); }
+	FORCEINLINE void Reset() { return Data.Reset(); }
+
+	FORCEINLINE TArray<FContextualAnimSceneActorData>::RangedForIteratorType      begin() { return Data.begin(); }
+	FORCEINLINE TArray<FContextualAnimSceneActorData>::RangedForConstIteratorType begin() const { return Data.begin(); }
+	FORCEINLINE TArray<FContextualAnimSceneActorData>::RangedForIteratorType      end() { return Data.end(); }
+	FORCEINLINE TArray<FContextualAnimSceneActorData>::RangedForConstIteratorType end() const { return Data.end(); }
+
+private:
+
+	UPROPERTY()
+	TArray<FContextualAnimSceneActorData> Data;
 };
 
 /** Delegate to notify external objects when this is scene is completed */
-DECLARE_DELEGATE_OneParam(FOnContextualAnimSceneEnded, class UContextualAnimSceneInstance*)
+DECLARE_DYNAMIC_MULTICAST_DELEGATE_OneParam(FOnContextualAnimSceneEnded, class UContextualAnimSceneInstance*, SceneInstance);
 
 /** Delegate to notify external objects when an actor join this scene */
-DECLARE_DELEGATE_TwoParams(FOnContextualAnimSceneActorJoined, class UContextualAnimSceneInstance*, AActor* Actor)
+DECLARE_DYNAMIC_MULTICAST_DELEGATE_TwoParams(FOnContextualAnimSceneActorJoined, class UContextualAnimSceneInstance*, SceneInstance, AActor*, Actor);
 
 /** Delegate to notify external objects when an actor left this scene */
-DECLARE_DELEGATE_TwoParams(FOnContextualAnimSceneActorLeft, class UContextualAnimSceneInstance*, AActor* Actor)
+DECLARE_DYNAMIC_MULTICAST_DELEGATE_TwoParams(FOnContextualAnimSceneActorLeft, class UContextualAnimSceneInstance*, SceneInstance, AActor*, Actor);
+
+/** Delegate to notify external objects about anim notify events */
+DECLARE_DYNAMIC_MULTICAST_DELEGATE_ThreeParams(FOnContextualAnimSceneNotify, class UContextualAnimSceneInstance*, SceneInstance, AActor*, Actor, FName, NotifyName);
 
 /** Instance of a contextual animation scene */
 UCLASS(BlueprintType, Blueprintable)
@@ -101,26 +148,29 @@ class CONTEXTUALANIMATION_API UContextualAnimSceneInstance : public UObject
 
 public:
 
-	/** Scene asset this instance was created from */
-	UPROPERTY()
-	TObjectPtr<const UContextualAnimSceneAsset> SceneAsset;
-
-	/** Map of roles to scene actor */
-	UPROPERTY()
-	TMap<FName, FContextualAnimSceneActorData> SceneActorMap;
-
-	/** List of alignment section to scene pivot */
-	TArray<TTuple<FName, FTransform>> AlignmentSectionToScenePivotList;
+	friend class UContextualAnimManager;
+	friend class FContextualAnimViewModel;
 
 	/** Delegate to notify external objects when this is scene is completed */
+	UPROPERTY(BlueprintAssignable)
 	FOnContextualAnimSceneEnded OnSceneEnded;
 
 	/** Delegate to notify external objects when an actor join */
+	UPROPERTY(BlueprintAssignable)
 	FOnContextualAnimSceneActorJoined OnActorJoined;
 
 	/** Delegate to notify external objects when an actor leave */
+	UPROPERTY(BlueprintAssignable)
 	FOnContextualAnimSceneActorLeft OnActorLeft;
 
+	/** Delegate to notify external objects when an animation hits a 'PlayMontageNotify' or 'PlayMontageNotifyWindow' begin */
+	UPROPERTY(BlueprintAssignable)
+	FOnContextualAnimSceneNotify OnNotifyBegin;
+
+	/** Delegate to notify external objects when an animation hits a 'PlayMontageNotify' or 'PlayMontageNotifyWindow' end */
+	UPROPERTY(BlueprintAssignable)
+	FOnContextualAnimSceneNotify OnNotifyEnd;
+
 	UContextualAnimSceneInstance(const FObjectInitializer& ObjectInitializer);
 
 	virtual UWorld* GetWorld() const override;
@@ -136,25 +186,23 @@ public:
 	/** Whether the supplied actor is part of this scene */
 	bool IsActorInThisScene(const AActor* Actor) const;
 
-	const FContextualAnimSceneActorData* FindSceneActorDataForActor(const AActor* Actor) const;
-
-	const FContextualAnimSceneActorData* FindSceneActorDataForRole(const FName& Role) const;
-
-	UFUNCTION(BlueprintCallable, Category = "Contextual Anim|Scene Instance")
-	float GetCurrentSectionTimeLeft() const;
-
-	UFUNCTION(BlueprintCallable, Category = "Contextual Anim|Scene Instance")
-	bool DidCurrentSectionLoop() const;
-
-	UFUNCTION(BlueprintCallable, Category = "Contextual Anim|Scene Instance")
-	float GetPositionInCurrentSection() const;
+	const UContextualAnimSceneAsset& GetSceneAsset() const { return *SceneAsset; }
+	const FContextualAnimSceneBindings& GetBindings() const { return Bindings; };
+	FContextualAnimSceneBindings& GetBindings() { return Bindings; };
+	const FContextualAnimSceneActorData* FindSceneActorDataByActor(const AActor* Actor) const { return Bindings.FindSceneActorDataByActor(Actor); }
+	const FContextualAnimSceneActorData* FindSceneActorDataByRole(const FName& Role) const { return Bindings.FindSceneActorDataByRole(Role); }
 
-	UFUNCTION(BlueprintNativeEvent, Category = "Contextual Anim|Scene Instance")
-	float GetResumePositionForSceneActor(const FContextualAnimSceneActorData& SceneActorData, int32 DesiredSectionIndex) const;
+#if WITH_EDITOR
+	const FContextualAnimSceneActorData* FindSceneActorDataByGuid(const FGuid& Guid) const { return Bindings.FindSceneActorDataByGuid(Guid); }
+#endif
 
 	UFUNCTION(BlueprintCallable, Category = "Contextual Anim|Scene Instance")
 	AActor* GetActorByRole(FName Role) const;
 
+	/** Extracts data from a ContextualAnimSceneActorData */
+	UFUNCTION(BlueprintPure, Category = "Contextual Anim|Scene Actor Data", meta = (NativeBreakFunc))
+	static void BreakContextualAnimSceneActorData(const FContextualAnimSceneActorData& SceneActorData, AActor*& Actor, UAnimMontage*& Montage, float& AnimTime, int32& CurrentSectionIndex, FName& CurrentSectionName);
+
 protected:
 
 	/** Tells the scene actor to join the scene (play animation) */
@@ -168,8 +216,6 @@ protected:
 	/** Helper function to set ignore collision between the supplied actor and all the other actors in this scene */
 	void SetIgnoreCollisionWithOtherActors(AActor* Actor, bool bValue) const;
 
-	void UpdateTransitions(float DeltaTime);
-
 	UFUNCTION()
 	void OnMontageBlendingOut(UAnimMontage* Montage, bool bInterrupted);
 
@@ -179,9 +225,14 @@ protected:
 	UFUNCTION()
 	void OnNotifyEndReceived(FName NotifyName, const FBranchingPointNotifyPayload& BranchingPointNotifyPayload);
 
-public:
+private:
 
-	/** Extracts data from a ContextualAnimSceneActorData */
-	UFUNCTION(BlueprintPure, Category = "Contextual Anim|Scene Actor Data", meta = (NativeBreakFunc))
-	static void BreakContextualAnimSceneActorData(const FContextualAnimSceneActorData& SceneActorData, AActor*& Actor, UAnimMontage*& Montage, float& AnimTime, int32& CurrentSectionIndex, FName& CurrentSectionName);
+	/** Scene asset this instance was created from */
+	UPROPERTY()
+	TObjectPtr<const UContextualAnimSceneAsset> SceneAsset;
+
+	UPROPERTY()
+	FContextualAnimSceneBindings Bindings;
+
+	TArray<TTuple<FName, FTransform>> AlignmentSectionToScenePivotList;
 };
\ No newline at end of file
diff --git a/Engine/Plugins/Experimental/Animation/ContextualAnimation/Source/ContextualAnimation/Public/ContextualAnimScenePivotProvider.h b/Engine/Plugins/Experimental/Animation/ContextualAnimation/Source/ContextualAnimation/Public/ContextualAnimScenePivotProvider.h
deleted file mode 100644
diff --git a/Engine/Plugins/Experimental/Animation/ContextualAnimation/Source/ContextualAnimation/Public/ContextualAnimSelectionCriterion.h b/Engine/Plugins/Experimental/Animation/ContextualAnimation/Source/ContextualAnimation/Public/ContextualAnimSelectionCriterion.h
new file mode 100644
--- /dev/null
+++ b/Engine/Plugins/Experimental/Animation/ContextualAnimation/Source/ContextualAnimation/Public/ContextualAnimSelectionCriterion.h
@@ -0,0 +1,64 @@
+// Copyright Epic Games, Inc. All Rights Reserved.
+
+#pragma once
+
+#include "CoreMinimal.h"
+#include "UObject/Object.h"
+#include "ContextualAnimTypes.h"
+#include "ContextualAnimSelectionCriterion.generated.h"
+
+// UContextualAnimSelectionCriterion
+//===========================================================================
+
+UCLASS(Abstract, BlueprintType, EditInlineNew)
+class CONTEXTUALANIMATION_API UContextualAnimSelectionCriterion : public UObject
+{
+	GENERATED_BODY()
+
+public:
+
+	UContextualAnimSelectionCriterion(const FObjectInitializer& ObjectInitializer);
+
+	class UContextualAnimSceneAsset* GetSceneAssetOwner() const;
+
+	virtual bool DoesQuerierPassCondition(const FContextualAnimPrimaryActorData& PrimaryActorData, const FContextualAnimQuerierData& QuerierData) const { return false; }
+};
+
+// UContextualAnimSelectionCriterion_TriggerArea
+//===========================================================================
+
+UCLASS()
+class CONTEXTUALANIMATION_API UContextualAnimSelectionCriterion_TriggerArea : public UContextualAnimSelectionCriterion
+{
+	GENERATED_BODY()
+
+public:
+
+	UPROPERTY(EditAnywhere, EditFixedSize, BlueprintReadWrite, Category = "Defaults", meta = (EditFixedOrder))
+	TArray<FVector> PolygonPoints;
+
+	UPROPERTY(EditAnywhere, BlueprintReadWrite, Category = "Defaults", meta = (ClampMin = "0", UIMin = "0"))
+	float Height = 100.f;
+
+	UContextualAnimSelectionCriterion_TriggerArea(const FObjectInitializer& ObjectInitializer);
+
+	virtual bool DoesQuerierPassCondition(const FContextualAnimPrimaryActorData& PrimaryActorData, const FContextualAnimQuerierData& QuerierData) const override;
+};
+
+// UContextualAnimSelectionCriterion_Facing
+//===========================================================================
+
+UCLASS()
+class CONTEXTUALANIMATION_API UContextualAnimSelectionCriterion_Facing : public UContextualAnimSelectionCriterion
+{
+	GENERATED_BODY()
+
+public:
+
+	UPROPERTY(EditAnywhere, BlueprintReadWrite, Category = "Defaults", meta = (ClampMin = "0", ClampMax = "180", UIMin = "0", UIMax = "180"))
+	float MaxAngle = 0.f;
+
+	UContextualAnimSelectionCriterion_Facing(const FObjectInitializer& ObjectInitializer) : Super(ObjectInitializer) {}
+
+	virtual bool DoesQuerierPassCondition(const FContextualAnimPrimaryActorData& PrimaryActorData, const FContextualAnimQuerierData& QuerierData) const override;
+};
\ No newline at end of file
diff --git a/Engine/Plugins/Experimental/Animation/ContextualAnimation/Source/ContextualAnimation/Public/ContextualAnimTypes.h b/Engine/Plugins/Experimental/Animation/ContextualAnimation/Source/ContextualAnimation/Public/ContextualAnimTypes.h
--- a/Engine/Plugins/Experimental/Animation/ContextualAnimation/Source/ContextualAnimation/Public/ContextualAnimTypes.h
+++ b/Engine/Plugins/Experimental/Animation/ContextualAnimation/Source/ContextualAnimation/Public/ContextualAnimTypes.h
@@ -4,13 +4,26 @@
 
 #include "Animation/AnimSequence.h"
 #include "Templates/SubclassOf.h"
+#include "GameplayTagContainer.h"
 #include "ContextualAnimTypes.generated.h"
 
 CONTEXTUALANIMATION_API DECLARE_LOG_CATEGORY_EXTERN(LogContextualAnim, Log, All);
 
 class AActor;
 class UAnimMontage;
-class UContextualAnimMetadata;
+class UContextualAnimSelectionCriterion;
+
+namespace UE 
+{
+	namespace ContextualAnim 
+	{
+		enum class EForEachResult : uint8
+		{
+			Break,
+			Continue,
+		};
+	}
+}
 
 /** Container for alignment tracks */
 USTRUCT()
@@ -42,20 +55,8 @@ struct CONTEXTUALANIMATION_API FContextualAnimAlignmentTrackContainer
 	FTransform ExtractTransformAtTime(const FName& TrackName, float Time) const;
 };
 
-USTRUCT()
-struct CONTEXTUALANIMATION_API FContextualAnimRoleBonePair
-{
-	GENERATED_BODY()
-
-	UPROPERTY()
-	FName RoleName = NAME_None;
-
-	UPROPERTY()
-	FName BoneName = NAME_None;
-};
-
 USTRUCT(BlueprintType)
-struct CONTEXTUALANIMATION_API FContextualAnimData
+struct CONTEXTUALANIMATION_API FContextualAnimTrack
 {
 	GENERATED_BODY()
 
@@ -76,19 +77,17 @@ struct CONTEXTUALANIMATION_API FContextualAnimData
 	UPROPERTY()
 	FContextualAnimAlignmentTrackContainer IKTargetData;
 
-	/** Lookup map to go from TrackName to Target Role and Parent Bone */
-	UPROPERTY()
-	TMap<FName, FContextualAnimRoleBonePair> IKTargetTrackLookupMap;
-
 	UPROPERTY(EditAnywhere, Instanced, Category = "Defaults")
-	UContextualAnimMetadata* Metadata = nullptr;
+	TArray<UContextualAnimSelectionCriterion*> SelectionCriteria;
 
 	UPROPERTY(EditAnywhere, BlueprintReadWrite, Category = "Defaults")
 	FTransform MeshToScene;
 
-	/** Index in the container of AnimData that owns us */
-	UPROPERTY()
-	int32 Index = INDEX_NONE;
+	UPROPERTY(EditAnywhere, BlueprintReadOnly, Category = "Defaults", meta = (GetOptions = "GetRoles"))
+	FName Role = NAME_None;
+
+	UPROPERTY(VisibleAnywhere, BlueprintReadOnly, Category = "Defaults")
+	int32 VariantIdx = INDEX_NONE;
 
 	float GetSyncTimeForWarpSection(int32 WarpSectionIndex) const;
 	float GetSyncTimeForWarpSection(const FName& WarpSectionName) const;
@@ -99,7 +98,9 @@ struct CONTEXTUALANIMATION_API FContextualAnimData
 
 	float FindBestAnimStartTime(const FVector& LocalLocation) const;
 
-	static const FContextualAnimData EmptyAnimData;
+	bool DoesQuerierPassSelectionCriteria(const FContextualAnimPrimaryActorData& PrimaryActorData, const FContextualAnimQuerierData& QuerierData) const;
+
+	static const FContextualAnimTrack EmptyTrack;
 };
 
 /** Defines when the actor should start playing the animation */
@@ -110,25 +111,6 @@ enum class EContextualAnimJoinRule : uint8
 	Late
 };
 
-USTRUCT(BlueprintType)
-struct FContextualAnimTransitionContainer
-{
-	GENERATED_BODY()
-
-	UPROPERTY(EditAnywhere, BlueprintReadWrite, Category = "Defaults")
-	TArray<FName> FromSections;
-
-	UPROPERTY(EditAnywhere, BlueprintReadWrite, Category = "Defaults")
-	FName ToSection = NAME_None;
-
-	UPROPERTY(EditAnywhere, Instanced, Category = "Defaults")
-	class UContextualAnimTransition* Transition = nullptr;
-
-	//@TODO: Remove this from here it should be on the EdMode
-	UPROPERTY(EditAnywhere, BlueprintReadWrite, Category = "Defaults")
-	bool bForceTransition = false;
-};
-
 UENUM(BlueprintType)
 enum class EContextualAnimIKTargetProvider : uint8
 {
@@ -188,59 +170,26 @@ struct CONTEXTUALANIMATION_API FContextualAnimIKTarget
 	UPROPERTY(EditAnywhere, BlueprintReadWrite, Category = "Defaults")
 	FTransform Transform;
 
-	FContextualAnimIKTarget(){}
+	FContextualAnimIKTarget() {}
 	FContextualAnimIKTarget(const FName& InGoalName, float InAlpha, const FTransform& InTransform)
-		: GoalName(InGoalName), Alpha(InAlpha), Transform(InTransform){}
+		: GoalName(InGoalName), Alpha(InAlpha), Transform(InTransform) {}
 
 	static const FContextualAnimIKTarget InvalidIKTarget;
 };
 
 USTRUCT(BlueprintType)
-struct CONTEXTUALANIMATION_API FContextualAnimTrackSettings
-{
-	GENERATED_BODY()
-
-	UPROPERTY(EditAnywhere, BlueprintReadWrite, Category = "Defaults")
-	TSubclassOf<AActor> PreviewActorClass;
-
-	UPROPERTY(EditAnywhere, BlueprintReadWrite, Category = "Defaults")
-	EContextualAnimJoinRule JoinRule = EContextualAnimJoinRule::Default;
-
-	UPROPERTY(EditAnywhere, BlueprintReadWrite, Category = "Defaults")
-	TArray<FContextualAnimIKTargetDefinition> IKTargetDefinitions;
-
-	UPROPERTY(EditAnywhere, BlueprintReadWrite, Category = "Defaults", AdvancedDisplay)
-	FTransform MeshToComponent = FTransform(FRotator(0.f, -90.f, 0.f));
-};
-
-USTRUCT(BlueprintType)
-struct FContextualAnimTrack
-{
-	GENERATED_BODY()
-
-	UPROPERTY(EditAnywhere, BlueprintReadWrite, Category = "Defaults")
-	FContextualAnimTrackSettings Settings;
-
-	UPROPERTY(EditAnywhere, BlueprintReadWrite, Category = "Defaults")
-	FContextualAnimData AnimData;
-};
-
-USTRUCT(BlueprintType)
-struct FContextualAnimCompositeTrack
+struct CONTEXTUALANIMATION_API FContextualAnimIKTargetDefContainer
 {
 	GENERATED_BODY()
 
 	UPROPERTY(EditAnywhere, BlueprintReadWrite, Category = "Defaults")
-	FContextualAnimTrackSettings Settings;
-
-	UPROPERTY(EditAnywhere, BlueprintReadWrite, Category = "Defaults")
-	TArray<FContextualAnimData> AnimDataContainer;
+	TArray<FContextualAnimIKTargetDefinition> IKTargetDefs;
 
-	FTransform GetRootTransformForAnimDataAtIndex(int32 Index) const;
+	static const FContextualAnimIKTargetDefContainer EmptyContainer;
 };
 
 USTRUCT(BlueprintType)
-struct FContextualAnimSceneBindings
+struct FContextualAnimStartSceneParams
 {
 	GENERATED_BODY()
 
@@ -248,13 +197,20 @@ struct FContextualAnimSceneBindings
 	UPROPERTY(EditAnywhere, BlueprintReadWrite, Category = "Defaults")
 	TMap<FName, AActor*> RoleToActorMap;
 
-	/** Desired AnimDataIndex. If INDEX_NONE the Manager will attempt to find the best AnimData to use */
+	/** Desired variant. If INDEX_NONE the Manager will attempt to find the best variant to use by running the selection criteria */
 	UPROPERTY(EditAnywhere, BlueprintReadWrite, Category = "Defaults")
-	int32 AnimDataIndex = INDEX_NONE;
+	int32 VariantIdx = INDEX_NONE;
 
-	/** Desired start time. Only relevant if AnimDataIndex != INDEX_NONE */
+	/** Desired start time */
 	UPROPERTY(EditAnywhere, BlueprintReadWrite, Category = "Defaults")
 	float AnimStartTime = 0.f;
+
+	void Reset()
+	{
+		RoleToActorMap.Reset();
+		VariantIdx = INDEX_NONE;
+		AnimStartTime = 0.f;
+	}
 };
 
 ///////////////////////////////////////////////////////////////////////
@@ -278,74 +234,39 @@ struct FContextualAnimQueryResult
 	float AnimStartTime = 0.f;
 
 	UPROPERTY(BlueprintReadWrite, Category = "Defaults")
-	int32 DataIndex = INDEX_NONE;
+	int32 VariantIdx = INDEX_NONE;
 
 	void Reset()
 	{
 		Animation.Reset();
 		EntryTransform = SyncTransform = FTransform::Identity;
 		AnimStartTime = 0.f;
-		DataIndex = INDEX_NONE;
+		VariantIdx = INDEX_NONE;
 	}
 
-	FORCEINLINE bool IsValid() const { return DataIndex != INDEX_NONE; }
+	FORCEINLINE bool IsValid() const { return VariantIdx != INDEX_NONE; }
 };
 
-/** Stores the parameters passed into query function */
 USTRUCT(BlueprintType)
-struct FContextualAnimQueryParams
+struct FContextualAnimQuerierData
 {
 	GENERATED_BODY()
 
 	UPROPERTY(EditAnywhere, BlueprintReadWrite, Category = "Defaults")
-	TWeakObjectPtr<const AActor> Querier;
-
-	UPROPERTY(EditAnywhere, BlueprintReadWrite, Category = "Defaults")
-	FTransform QueryTransform;
-
-	UPROPERTY(EditAnywhere, BlueprintReadWrite, Category = "Defaults")
-	bool bComplexQuery = false;
-
-	UPROPERTY(EditAnywhere, BlueprintReadWrite, Category = "Defaults")
-	bool bFindAnimStartTime = false;
-
-	FContextualAnimQueryParams() {}
-
-	FContextualAnimQueryParams(const AActor* InQuerier, bool bInComplexQuery, bool bInFindAnimStartTime)
-		: Querier(InQuerier), bComplexQuery(bInComplexQuery), bFindAnimStartTime(bInFindAnimStartTime) {}
-
-	FContextualAnimQueryParams(const FTransform& InQueryTransform, bool bInComplexQuery, bool bInFindAnimStartTime)
-		: QueryTransform(InQueryTransform), bComplexQuery(bInComplexQuery), bFindAnimStartTime(bInFindAnimStartTime) {}
-};
-
-USTRUCT(BlueprintType)
-struct FContextualAnimQuerier
-{
-	GENERATED_BODY()
+	FTransform Transform;
 
 	UPROPERTY(EditAnywhere, BlueprintReadWrite, Category = "Defaults")
-	TWeakObjectPtr<const AActor> Actor = nullptr;
+	FVector Velocity = FVector::ZeroVector;
 
 	UPROPERTY(EditAnywhere, BlueprintReadWrite, Category = "Defaults")
-	FTransform Transform;
-
-	FContextualAnimQuerier(){}
-	FContextualAnimQuerier(const AActor* InActor) : Actor(InActor) {}
-	FContextualAnimQuerier(const FTransform& InTransform) : Transform(InTransform) {}
+	FGameplayTagContainer GameplayTags;
 };
 
 USTRUCT(BlueprintType)
-struct FContextualAnimQueryContext
+struct FContextualAnimPrimaryActorData
 {
 	GENERATED_BODY()
 
-	UPROPERTY(EditAnywhere, BlueprintReadWrite, Category = "Defaults")
-	TWeakObjectPtr<const AActor> Actor = nullptr;
-
 	UPROPERTY(EditAnywhere, BlueprintReadWrite, Category = "Defaults")
 	FTransform Transform;
-
-	FContextualAnimQueryContext() {}
-	FContextualAnimQueryContext(const AActor* InActor) : Actor(InActor) {}
-	FContextualAnimQueryContext(const FTransform& InTransform) : Transform(InTransform) {}
 };
\ No newline at end of file
diff --git a/Engine/Plugins/Experimental/Animation/ContextualAnimation/Source/ContextualAnimation/Public/ContextualAnimUtilities.h b/Engine/Plugins/Experimental/Animation/ContextualAnimation/Source/ContextualAnimation/Public/ContextualAnimUtilities.h
--- a/Engine/Plugins/Experimental/Animation/ContextualAnimation/Source/ContextualAnimation/Public/ContextualAnimUtilities.h
+++ b/Engine/Plugins/Experimental/Animation/ContextualAnimation/Source/ContextualAnimation/Public/ContextualAnimUtilities.h
@@ -11,6 +11,7 @@ class UContextualAnimSceneAsset;
 class USkeletalMeshComponent;
 class UAnimInstance;
 class AActor;
+class FPrimitiveDrawInterface;
 struct FAnimMontageInstance;
 
 UCLASS()
@@ -42,7 +43,7 @@ public:
 
 	static void DrawDebugPose(const UWorld* World, const UAnimSequenceBase* Animation, float Time, const FTransform& LocalToWorldTransform, const FColor& Color, float LifeTime, float Thickness);
 	
-	static void DrawDebugScene(const UWorld* World, const UContextualAnimSceneAsset* SceneAsset, int32 AnimDataIndex, float Time, const FTransform& ToWorldTransform, const FColor& Color, float LifeTime, float Thickness);
+	static void DrawDebugScene(const UWorld* World, const UContextualAnimSceneAsset* SceneAsset, int32 VariantIdx, float Time, const FTransform& ToWorldTransform, const FColor& Color, float LifeTime, float Thickness);
 
 	static USkeletalMeshComponent* TryGetSkeletalMeshComponent(const AActor* Actor);
 
@@ -59,4 +60,5 @@ public:
 	UFUNCTION(BlueprintCallable, Category = "Contextual Anim|Utilities", meta = (DisplayName = "GetSectionLength"))
 	static float BP_Montage_GetSectionLength(const UAnimMontage* Montage, int32 SectionIndex);
 
+	static void DrawSector(FPrimitiveDrawInterface& PDI, const FVector& Origin, const FVector& Direction, float MinDistance, float MaxDistance, float MinAngle, float MaxAngle, const FLinearColor& Color, uint8 DepthPriority, float Thickness);
 };
\ No newline at end of file
diff --git a/Engine/Plugins/Experimental/Animation/ContextualAnimation/Source/ContextualAnimationEditor/ContextualAnimationEditor.Build.cs b/Engine/Plugins/Experimental/Animation/ContextualAnimation/Source/ContextualAnimationEditor/ContextualAnimationEditor.Build.cs
--- a/Engine/Plugins/Experimental/Animation/ContextualAnimation/Source/ContextualAnimationEditor/ContextualAnimationEditor.Build.cs
+++ b/Engine/Plugins/Experimental/Animation/ContextualAnimation/Source/ContextualAnimationEditor/ContextualAnimationEditor.Build.cs
@@ -27,7 +27,8 @@ public class ContextualAnimationEditor : ModuleRules
 				"AIModule",
 				"NavigationSystem",
 				"ContextualAnimation",
-				"MotionWarping",	
+				"MotionWarping",
+				"GameplayTags"
 			}
 			);
 			
@@ -58,7 +59,9 @@ public class ContextualAnimationEditor : ModuleRules
 				"MovieSceneTracks",
 				"MovieSceneTools",
 				"ClassViewer",
-				"Persona"
+				"Persona",
+				"RenderCore",
+				"BlueprintGraph"
 			}
 			);
 		
diff --git a/Engine/Plugins/Experimental/Animation/ContextualAnimation/Source/ContextualAnimationEditor/Private/ContextualAnimAssetEditorCommands.cpp b/Engine/Plugins/Experimental/Animation/ContextualAnimation/Source/ContextualAnimationEditor/Private/ContextualAnimAssetEditorCommands.cpp
--- a/Engine/Plugins/Experimental/Animation/ContextualAnimation/Source/ContextualAnimationEditor/Private/ContextualAnimAssetEditorCommands.cpp
+++ b/Engine/Plugins/Experimental/Animation/ContextualAnimation/Source/ContextualAnimationEditor/Private/ContextualAnimAssetEditorCommands.cpp
@@ -7,6 +7,13 @@
 void FContextualAnimAssetEditorCommands::RegisterCommands()
 {
 	UI_COMMAND(ResetPreviewScene, "Reset Scene", "Reset Scene.", EUserInterfaceActionType::Button, FInputChord());
+
+	UI_COMMAND(ShowIKTargetsDrawSelected, "Selected Actor Only", "Show IK Targets for the selected actor", EUserInterfaceActionType::RadioButton, FInputChord());
+	UI_COMMAND(ShowIKTargetsDrawAll, "All Actors", "Show IK Targets for all the actors", EUserInterfaceActionType::RadioButton, FInputChord());
+	UI_COMMAND(ShowIKTargetsDrawNone, "None", "Hide IK Targets", EUserInterfaceActionType::RadioButton, FInputChord());
+
+	UI_COMMAND(Simulate, "Simulate", "Simulate Mode", EUserInterfaceActionType::RadioButton, FInputChord());
+
 }
 
 #undef LOCTEXT_NAMESPACE
diff --git a/Engine/Plugins/Experimental/Animation/ContextualAnimation/Source/ContextualAnimationEditor/Private/ContextualAnimAssetEditorToolkit.cpp b/Engine/Plugins/Experimental/Animation/ContextualAnimation/Source/ContextualAnimationEditor/Private/ContextualAnimAssetEditorToolkit.cpp
--- a/Engine/Plugins/Experimental/Animation/ContextualAnimation/Source/ContextualAnimationEditor/Private/ContextualAnimAssetEditorToolkit.cpp
+++ b/Engine/Plugins/Experimental/Animation/ContextualAnimation/Source/ContextualAnimationEditor/Private/ContextualAnimAssetEditorToolkit.cpp
@@ -15,8 +15,11 @@
 #include "MovieScene.h"
 #include "Framework/MultiBox/MultiBoxBuilder.h"
 #include "AdvancedPreviewSceneModule.h"
+#include "IStructureDetailsView.h"
 #include "PropertyEditorModule.h"
 #include "Modules/ModuleManager.h"
+#include "ContextualAnimEditorTypes.h"
+#include "Widgets/Input/SButton.h"
 
 #define LOCTEXT_NAMESPACE "ContextualAnimAssetEditorToolkit"
 
@@ -51,11 +54,6 @@ UContextualAnimSceneAsset* FContextualAnimAssetEditorToolkit::GetSceneAsset() co
 	return ViewModel.IsValid() ? ViewModel->GetSceneAsset() : nullptr;
 }
 
-UContextualAnimPreviewManager* FContextualAnimAssetEditorToolkit::GetPreviewManager() const
-{
-	return ViewModel.IsValid() ? ViewModel->GetPreviewManager() : nullptr;
-}
-
 void FContextualAnimAssetEditorToolkit::ResetPreviewScene()
 {
 	ViewModel->RefreshSequencerTracks();
@@ -98,7 +96,7 @@ void FContextualAnimAssetEditorToolkit::InitAssetEditor(const EToolkitMode::Type
 	EditingAssetWidget->OnFinishedChangingProperties().AddSP(this, &FContextualAnimAssetEditorToolkit::OnFinishedChangingProperties);
 
 	// Define Editor Layout
-	const TSharedRef<FTabManager::FLayout> StandaloneDefaultLayout = FTabManager::NewLayout("Standalone_ContextualAnimAnimEditor_Layout_v0.07")
+	const TSharedRef<FTabManager::FLayout> StandaloneDefaultLayout = FTabManager::NewLayout("Standalone_ContextualAnimAnimEditor_Layout_v0.09")
 		->AddArea
 		(
 			FTabManager::NewPrimaryArea()->SetOrientation(Orient_Vertical)
@@ -163,6 +161,25 @@ void FContextualAnimAssetEditorToolkit::BindCommands()
 		Commands.ResetPreviewScene,
 		FExecuteAction::CreateSP(this, &FContextualAnimAssetEditorToolkit::ResetPreviewScene),
 		EUIActionRepeatMode::RepeatDisabled);
+
+	ToolkitCommands->MapAction(
+		Commands.Simulate,
+		FExecuteAction::CreateSP(this, &FContextualAnimAssetEditorToolkit::ToggleSimulateMode),
+		FCanExecuteAction(),
+		FIsActionChecked::CreateSP(this, &FContextualAnimAssetEditorToolkit::IsSimulateModeActive));
+}
+
+void FContextualAnimAssetEditorToolkit::ToggleSimulateMode()
+{
+	if(ViewModel.IsValid())
+	{
+		ViewModel->ToggleSimulateMode();
+	}
+}
+
+bool FContextualAnimAssetEditorToolkit::IsSimulateModeActive() const
+{
+	return (ViewModel.IsValid()) ? ViewModel->IsSimulateModeActive() : false;
 }
 
 void FContextualAnimAssetEditorToolkit::ExtendToolbar()
@@ -181,16 +198,128 @@ void FContextualAnimAssetEditorToolkit::ExtendToolbar()
 
 void FContextualAnimAssetEditorToolkit::FillToolbar(FToolBarBuilder& ToolbarBuilder)
 {
-	ToolbarBuilder.BeginSection("Reset");
+	ToolbarBuilder.AddToolBarButton(
+		FContextualAnimAssetEditorCommands::Get().ResetPreviewScene,
+		NAME_None,
+		TAttribute<FText>(),
+		TAttribute<FText>(),
+		FSlateIcon(FAppStyle::Get().GetStyleSetName(), "Icons.Refresh")
+	);
+
+	ToolbarBuilder.AddComboButton(
+		FUIAction(),
+		FOnGetContent::CreateSP(this, &FContextualAnimAssetEditorToolkit::BuildVariantsMenu),
+		LOCTEXT("Variants_Label", "Variants"),
+		FText::GetEmpty(),
+		FSlateIcon()
+	);
+
+	ToolbarBuilder.AddToolBarButton(
+		FContextualAnimAssetEditorCommands::Get().Simulate,
+		NAME_None,
+		TAttribute<FText>(),
+		TAttribute<FText>(),
+		FSlateIcon()
+	);
+}
+
+TSharedRef<SWidget> FContextualAnimAssetEditorToolkit::BuildVariantsMenu()
+{
+	const bool bShouldCloseWindowAfterMenuSelection = true;
+	FMenuBuilder MenuBuilder(bShouldCloseWindowAfterMenuSelection, GetToolkitCommands());
+
 	{
-		ToolbarBuilder.AddToolBarButton(
-			FContextualAnimAssetEditorCommands::Get().ResetPreviewScene,
-			NAME_None,
-			TAttribute<FText>(),
-			TAttribute<FText>(),
-			FSlateIcon(FAppStyle::Get().GetStyleSetName(), "Icons.Refresh"));
+		MenuBuilder.BeginSection(NAME_None, LOCTEXT("CreateNewVariant", "New Variant"));
+
+		MenuBuilder.AddSubMenu(
+			LOCTEXT("NewVariant", "New Variant"),
+			FText::GetEmpty(),
+			FNewMenuDelegate::CreateLambda([this](FMenuBuilder& MenuBuilder){
+			
+				FPropertyEditorModule& PropertyModule = FModuleManager::LoadModuleChecked<FPropertyEditorModule>("PropertyEditor");
+
+				FDetailsViewArgs Args;
+				Args.bHideSelectionTip = true;
+				Args.bAllowSearch = false;
+				Args.bAllowFavoriteSystem = false;
+
+				NewVariantWidgetStruct = MakeShared<FStructOnScope>(FContextualAnimNewVariantParams::StaticStruct());
+				FContextualAnimNewVariantParams* Params = (FContextualAnimNewVariantParams*)NewVariantWidgetStruct->GetStructMemory();
+
+				const TArray<FName> Roles = ViewModel->GetSceneAsset()->GetRoles();
+				for (FName Role : Roles)
+				{
+					FContextualAnimNewVariantData Entry;
+					Entry.RoleName = Role;
+					Params->Data.Add(Entry);
+				}
+
+				TSharedRef<IStructureDetailsView> StructureDetailsView = PropertyModule.CreateStructureDetailView(Args, FStructureDetailsViewArgs(), NewVariantWidgetStruct);
+
+				MenuBuilder.AddWidget(
+					SNew(SBox)
+					.MinDesiredWidth(500.0f)
+					.MaxDesiredWidth(500.f)
+					.MaxDesiredHeight(400.0f)
+					[
+						SNew(SVerticalBox)
+						+ SVerticalBox::Slot()
+						.AutoHeight()
+						.HAlign(HAlign_Fill)
+						[
+							StructureDetailsView->GetWidget().ToSharedRef()
+						]
+						+ SVerticalBox::Slot()
+						.AutoHeight()
+						.HAlign(HAlign_Fill)
+						[
+							SNew(SButton)
+							.ContentPadding(3)
+							.VAlign(VAlign_Center)
+							.HAlign(HAlign_Center)
+							.OnClicked_Lambda([this]()
+							{
+								FContextualAnimNewVariantParams* Params = (FContextualAnimNewVariantParams*)NewVariantWidgetStruct->GetStructMemory();
+								check(Params);
+
+								ViewModel->AddNewVariant(*Params);
+
+								FSlateApplication::Get().DismissAllMenus();
+
+								return FReply::Handled();
+							})
+						.Text(LOCTEXT("OK", "OK"))
+						]
+					],
+					FText(), true, false);
+			}),
+			false,
+			FSlateIcon()
+		);
+
+		MenuBuilder.EndSection();
 	}
-	ToolbarBuilder.EndSection();
+
+	{
+		MenuBuilder.BeginSection(NAME_None, LOCTEXT("Variants_Label", "Variants"));
+		{
+			const UContextualAnimSceneAsset* SceneAsset = GetSceneAsset();
+			const int32 TotalVariants = SceneAsset->GetTotalVariants();
+			for (int32 Idx = 0; Idx < TotalVariants; Idx++)
+			{
+				MenuBuilder.AddMenuEntry(
+					FText::FromString(FString::Printf(TEXT("%d"), Idx)),
+					FText::GetEmpty(),
+					FSlateIcon(),
+					FUIAction(FExecuteAction::CreateLambda([this, Idx]() {
+						ViewModel->SetActiveSceneVariantIdx(Idx);
+					})));
+			}
+		}
+		MenuBuilder.EndSection();
+	}
+
+	return MenuBuilder.MakeWidget();
 }
 
 void FContextualAnimAssetEditorToolkit::RegisterTabSpawners(const TSharedRef<FTabManager>& InTabManager)
@@ -332,6 +461,11 @@ TSharedRef<SDockTab> FContextualAnimAssetEditorToolkit::SpawnTab_PreviewSettings
 
 void FContextualAnimAssetEditorToolkit::OnFinishedChangingProperties(const FPropertyChangedEvent& PropertyChangedEvent)
 {
+	const FName PropertyName = (PropertyChangedEvent.Property != nullptr) ? PropertyChangedEvent.Property->GetFName() : NAME_None;
+	const FName MemberPropertyName = (PropertyChangedEvent.MemberProperty != nullptr) ? PropertyChangedEvent.MemberProperty->GetFName() : NAME_None;
+
+	UE_LOG(LogContextualAnim, Log, TEXT("FContextualAnimAssetEditorToolkit::OnFinishedChangingProperties MemberPropertyName: %s PropertyName: %s"), 
+		*MemberPropertyName.ToString(), *PropertyName.ToString());
 }
 
 #undef LOCTEXT_NAMESPACE
diff --git a/Engine/Plugins/Experimental/Animation/ContextualAnimation/Source/ContextualAnimationEditor/Private/ContextualAnimEdMode.cpp b/Engine/Plugins/Experimental/Animation/ContextualAnimation/Source/ContextualAnimationEditor/Private/ContextualAnimEdMode.cpp
--- a/Engine/Plugins/Experimental/Animation/ContextualAnimation/Source/ContextualAnimationEditor/Private/ContextualAnimEdMode.cpp
+++ b/Engine/Plugins/Experimental/Animation/ContextualAnimation/Source/ContextualAnimationEditor/Private/ContextualAnimEdMode.cpp
@@ -1,11 +1,26 @@
 // Copyright Epic Games, Inc. All Rights Reserved.
 
 #include "ContextualAnimEdMode.h"
-#include "EditorViewportClient.h"
 #include "EngineUtils.h"
+#include "ContextualAnimSceneAsset.h"
+#include "ContextualAnimSelectionCriterion.h"
+#include "ContextualAnimUtilities.h"
+#include "ContextualAnimSceneInstance.h"
+#include "ContextualAnimViewportClient.h"
+#include "ContextualAnimAssetEditorToolkit.h"
+#include "ContextualAnimViewModel.h"
+#include "Animation/AnimMontage.h"
+#include "AnimNotifyState_IKWindow.h"
+#include "SkeletalDebugRendering.h"
+
+#include "GameFramework/Character.h"
+#include "Components/CapsuleComponent.h"
+#include "Components/SkeletalMeshComponent.h"
 
 const FEditorModeID FContextualAnimEdMode::EdModeId = TEXT("ContextualAnimEdMode");
 
+IMPLEMENT_HIT_PROXY(HSelectionCriterionHitProxy, HHitProxy);
+
 FContextualAnimEdMode::FContextualAnimEdMode()
 {
 }
@@ -14,20 +29,207 @@ FContextualAnimEdMode::~FContextualAnimEdMode()
 {
 }
 
-bool FContextualAnimEdMode::HandleClick(FEditorViewportClient* InViewportClient, HHitProxy* HitProxy, const FViewportClick& Click)
+void FContextualAnimEdMode::Tick(FEditorViewportClient* ViewportClient, float DeltaTime)
 {
-	bool bHandled = false;
+	FEdMode::Tick(ViewportClient, DeltaTime);
 
-	if (HitProxy != nullptr && HitProxy->IsA(HActor::StaticGetType()))
+	if (ViewModel && ViewModel->IsSimulateModeActive())
 	{
-		HActor* ActorHitProxy = static_cast<HActor*>(HitProxy);
+		// @TODO: Two things here are temporal. First, we have to click on the actor that we want to control, and second, we can only control characters. 
+		// We need to find a better work flow for this, but its good enough for now as a prototype
+		if (!ViewportClient->Viewport->KeyState(EKeys::RightMouseButton))
+		{
+			if (ViewportClient->Viewport->KeyState(EKeys::W) || ViewportClient->Viewport->KeyState(EKeys::S))
+			{
+				if (ACharacter* Character = Cast<ACharacter>(SelectedActor.Get()))
+				{
+					const FVector WorldDirection = FRotationMatrix(Character->GetActorRotation()).GetScaledAxis(EAxis::X);
+					Character->AddMovementInput(WorldDirection, ViewportClient->Viewport->KeyState(EKeys::W) ? 1.f : -1.f);
+				}
+			}
 
-		SelectedActor = ActorHitProxy->Actor;
+			if (ViewportClient->Viewport->KeyState(EKeys::A) || ViewportClient->Viewport->KeyState(EKeys::D))
+			{
+				if (ACharacter* Character = Cast<ACharacter>(SelectedActor.Get()))
+				{
+					const FVector WorldDirection = FRotationMatrix(Character->GetActorRotation()).GetScaledAxis(EAxis::Y);
+					Character->AddMovementInput(WorldDirection, ViewportClient->Viewport->KeyState(EKeys::D) ? 1.f : -1.f);
+				}
+			}
+		}
+	}
+}
+
+void FContextualAnimEdMode::Render(const FSceneView* View, FViewport* Viewport, FPrimitiveDrawInterface* PDI)
+{
+	FEdMode::Render(View, Viewport, PDI);
 
-		bHandled = true;
+	// @TODO: This should not be initialized here
+	FContextualAnimViewportClient* ViewportClient = static_cast<FContextualAnimViewportClient*>(Viewport->GetClient());
+	if(!ViewModel)
+	{
+		ViewModel = ViewportClient->GetAssetEditorToolkit()->GetViewModel();
 	}
 
-	return bHandled;
+	if (ViewModel)
+	{
+		// Draw a circle below the selected actor if it is a character to let the user know that it is a controllable actor during simulate mode
+		if(ViewModel->IsSimulateModeActive() && SelectedActor.IsValid())
+		{
+			if (const ACharacter* Character = Cast<ACharacter>(SelectedActor.Get()))
+			{
+				const float CapsuleHalfHeight = Character->GetCapsuleComponent()->GetScaledCapsuleHalfHeight();
+				const FVector Location = Character->GetActorLocation() - Character->GetActorQuat().GetUpVector() * CapsuleHalfHeight;
+				DrawCircle(PDI, Location, FVector(1, 0, 0), FVector(0, 1, 0), FColor::Red, Character->GetCapsuleComponent()->GetScaledCapsuleRadius(), 12, SDPG_World, 1.f);
+			}
+		}
+
+		if (const UContextualAnimSceneInstance* SceneInstance = ViewModel->GetSceneInstance())
+		{
+			const UContextualAnimSceneAsset* SceneAsset = ViewModel->GetSceneAsset();
+			const int32 VariantIdx = ViewModel->GetActiveSceneVariantIdx();
+
+			// Draw Scene Pivots
+			if (SceneAsset->Variants.IsValidIndex(VariantIdx))
+			{
+				for(const FTransform& ScenePivot : SceneAsset->Variants[VariantIdx].ScenePivots)
+				{
+					DrawCoordinateSystem(PDI, ScenePivot.GetLocation(), ScenePivot.Rotator(), 50.f, SDPG_Foreground);
+				}
+			}
+
+			FTransform PrimaryActorTransform = FTransform::Identity;
+			if (const FContextualAnimSceneActorData* PreviewData = SceneInstance->FindSceneActorDataByRole(SceneAsset->GetPrimaryRole()))
+			{
+				PrimaryActorTransform = PreviewData->GetActor()->GetTransform();
+			}
+
+			const EShowIKTargetsDrawMode IKTargetsDrawMode = ViewportClient->GetShowIKTargetsDrawMode();
+
+			for (const FContextualAnimSceneActorData& SceneActorData : SceneInstance->GetBindings())
+			{
+				// Draw IK Targets
+				if (IKTargetsDrawMode == EShowIKTargetsDrawMode::All || (IKTargetsDrawMode == EShowIKTargetsDrawMode::Selected && SceneActorData.GetActor() == SelectedActor.Get()))
+				{
+					DrawIKTargetsForSceneActor(*PDI, SceneActorData);
+				}
+
+				// Draw Selection Criteria
+				if (const FContextualAnimTrack* AnimTrack = SceneAsset->GetAnimTrack(SceneActorData.GetRole(), VariantIdx))
+				{
+					for (int32 CriterionIdx = 0; CriterionIdx < AnimTrack->SelectionCriteria.Num(); CriterionIdx++)
+					{
+						if(const UContextualAnimSelectionCriterion* Criterion = AnimTrack->SelectionCriteria[CriterionIdx])
+						{
+							FLinearColor DrawColor = FLinearColor::White;
+
+							{
+								FContextualAnimQuerierData QuerierData;
+								QuerierData.Transform = SceneActorData.GetActor()->GetActorTransform();
+								QuerierData.Velocity = SceneActorData.GetActor()->GetVelocity();
+
+								FContextualAnimPrimaryActorData PrimaryActorData;
+								PrimaryActorData.Transform = PrimaryActorTransform;
+
+								if (Criterion->DoesQuerierPassCondition(PrimaryActorData, QuerierData))
+								{
+									DrawColor = FLinearColor::Green;
+								}
+							}
+
+							//@TODO: Each SelectionCriterion should implement this, and here we should just call "Criterion->Draw()"
+							if (const UContextualAnimSelectionCriterion_TriggerArea* Spatial = Cast<UContextualAnimSelectionCriterion_TriggerArea>(Criterion))
+							{
+								const float HalfHeight = Spatial->Height / 2.f;
+								const int32 LastIndex = Spatial->PolygonPoints.Num() - 1;
+								for (int32 Idx = 0; Idx <= LastIndex; Idx++)
+								{
+									const FVector P0 = PrimaryActorTransform.TransformPositionNoScale(Spatial->PolygonPoints[Idx]);
+									const FVector P1 = PrimaryActorTransform.TransformPositionNoScale(Spatial->PolygonPoints[Idx == LastIndex ? 0 : Idx + 1]);
+
+									PDI->DrawLine(P0, P1, DrawColor, SDPG_Foreground, 2.f);
+
+									PDI->DrawLine(P0 + FVector::UpVector * Spatial->Height, P1 + FVector::UpVector * Spatial->Height, DrawColor, SDPG_Foreground, 2.f);
+
+									PDI->DrawLine(P0, P0 + FVector::UpVector * Spatial->Height, DrawColor, SDPG_Foreground, 2.f);
+
+									PDI->SetHitProxy(new HSelectionCriterionHitProxy(FSelectionCriterionHitProxyData(SceneActorData.GetRole(), VariantIdx, CriterionIdx, Idx)));
+									PDI->DrawPoint(P0, FLinearColor::Black, 15.f, SDPG_Foreground);
+									PDI->SetHitProxy(nullptr);
+
+									PDI->SetHitProxy(new HSelectionCriterionHitProxy(FSelectionCriterionHitProxyData(SceneActorData.GetRole(), VariantIdx, CriterionIdx, Idx + 4)));
+									PDI->DrawPoint(P0 + FVector::UpVector * Spatial->Height, FLinearColor::Black, 15.f, SDPG_Foreground);
+									PDI->SetHitProxy(nullptr);
+								}
+							}
+							else if (const UContextualAnimSelectionCriterion_Facing* Facing = Cast<UContextualAnimSelectionCriterion_Facing>(Criterion))
+							{
+								const FTransform Transform = AnimTrack->GetAlignmentTransformAtEntryTime() * PrimaryActorTransform;
+								UContextualAnimUtilities::DrawSector(*PDI, Transform.GetLocation(), Transform.GetRotation().GetForwardVector(), 0.f, 30.f, -Facing->MaxAngle, Facing->MaxAngle, DrawColor, SDPG_World, 1.f);
+							}
+						}
+					}
+				}
+			}
+		}
+	}
+}
+
+void FContextualAnimEdMode::DrawIKTargetsForSceneActor(FPrimitiveDrawInterface& PDI, const FContextualAnimSceneActorData& SceneActorData) const
+{
+	for (const FContextualAnimIKTargetDefinition& IKTargetDef : SceneActorData.GetIKTargetDefs().IKTargetDefs)
+	{
+		if (const FContextualAnimSceneActorData* TargetSceneActorData = SceneActorData.GetSceneInstance().FindSceneActorDataByRole(IKTargetDef.TargetRoleName))
+		{
+			if (USkeletalMeshComponent* TargetSkelMeshComp = TargetSceneActorData->GetSkeletalMeshComponent())
+			{
+				const float Alpha = UAnimNotifyState_IKWindow::GetIKAlphaValue(IKTargetDef.GoalName, SceneActorData.GetAnimMontageInstance());
+
+				if (IKTargetDef.Provider == EContextualAnimIKTargetProvider::Bone)
+				{
+					//@TODO:...
+				}
+				else if (IKTargetDef.Provider == EContextualAnimIKTargetProvider::Autogenerated)
+				{
+					const FTransform ParentTransform = TargetSkelMeshComp->GetSocketTransform(IKTargetDef.TargetBoneName);
+
+					const FTransform TargetTransform = SceneActorData.GetAnimTrack().IKTargetData.ExtractTransformAtTime(IKTargetDef.GoalName, SceneActorData.GetAnimTime()) * ParentTransform;
+
+					FLinearColor Color = Alpha > 0.f ? FLinearColor(FColor::MakeRedToGreenColorFromScalar(Alpha)) : FLinearColor::White;
+
+					FVector Start = ParentTransform.GetLocation();
+					FVector End = TargetTransform.GetLocation();
+
+					const float Radius = 1.f;
+					SkeletalDebugRendering::DrawWireBone(&PDI, Start, End, Color, SDPG_Foreground, Radius);
+					SkeletalDebugRendering::DrawAxes(&PDI, FTransform(End), SDPG_Foreground, 0.f, Radius);
+				}
+			}
+		}
+	}
+}
+
+bool FContextualAnimEdMode::HandleClick(FEditorViewportClient* InViewportClient, HHitProxy* HitProxy, const FViewportClick& Click)
+{
+	SelectedActor.Reset();
+	SelectedSelectionCriterionData.Reset();
+
+	if (HitProxy != nullptr)
+	{
+		if (HitProxy->IsA(HActor::StaticGetType()))
+		{
+			HActor* ActorHitProxy = static_cast<HActor*>(HitProxy);
+			SelectedActor = ActorHitProxy->Actor;
+			return true;
+		}
+		else if (HitProxy->IsA(HSelectionCriterionHitProxy::StaticGetType()))
+		{
+			SelectedSelectionCriterionData = static_cast<HSelectionCriterionHitProxy*>(HitProxy)->Data;
+			return true;
+		}
+	}
+
+	return false; // unhandled
 }
 
 bool FContextualAnimEdMode::GetHitResultUnderCursor(FHitResult& OutHitResult, FEditorViewportClient* InViewportClient, const FViewportClick& Click) const
@@ -55,20 +257,66 @@ bool FContextualAnimEdMode::EndTracking(FEditorViewportClient* InViewportClient,
 
 bool FContextualAnimEdMode::InputDelta(FEditorViewportClient* InViewportClient, FViewport* InViewport, FVector& InDrag, FRotator& InRot, FVector& InScale)
 {
-	bool bHandled = false;
-
 	EAxisList::Type CurrentAxis = InViewportClient->GetCurrentWidgetAxis();
 
 	if (CurrentAxis != EAxisList::None)
 	{
-		if (SelectedActor.IsValid())
+		if (ViewModel)
 		{
-			SelectedActor->AddActorWorldTransform(FTransform(InRot, InDrag));
-			bHandled = true;
+			if (SelectedSelectionCriterionData.IsValid())
+			{
+				const UContextualAnimSceneAsset* SceneAsset = ViewModel->GetSceneAsset();
+				if (const FContextualAnimTrack* AnimTrack = SceneAsset->GetAnimTrack(SelectedSelectionCriterionData.RoleName, SelectedSelectionCriterionData.VariantIdx))
+				{
+					if (UContextualAnimSelectionCriterion_TriggerArea* Spatial = Cast<UContextualAnimSelectionCriterion_TriggerArea>(AnimTrack->SelectionCriteria[SelectedSelectionCriterionData.CriterionIdx]))
+					{
+						FMatrix WidgetCoordSystem;
+						GetCustomDrawingCoordinateSystem(WidgetCoordSystem, nullptr);
+						InDrag = WidgetCoordSystem.InverseTransformVector(InDrag);
+
+						FVector& Point = Spatial->PolygonPoints[SelectedSelectionCriterionData.DataIdx >= 4 ? SelectedSelectionCriterionData.DataIdx - 4 : SelectedSelectionCriterionData.DataIdx];
+						Point.X += InDrag.X;
+						Point.Y += InDrag.Y;
+
+						if (InDrag.Z != 0.f)
+						{
+							if (SelectedSelectionCriterionData.DataIdx < 4)
+							{
+								for (int32 Idx = 0; Idx < Spatial->PolygonPoints.Num(); Idx++)
+								{
+									Spatial->PolygonPoints[Idx].Z += InDrag.Z;
+								}
+
+								Spatial->Height = FMath::Max(Spatial->Height - InDrag.Z, 0.f);
+							}
+							else
+							{
+								Spatial->Height = FMath::Max(Spatial->Height + InDrag.Z, 0.f);
+							}
+						}
+
+						return true;
+					}
+				}
+			}
 		}
 	}
 
-	return bHandled;
+	return false; // unhandled
+}
+
+bool FContextualAnimEdMode::InputKey(FEditorViewportClient* ViewportClient, FViewport* Viewport, FKey Key, EInputEvent Event)
+{
+	if(ViewModel && ViewModel->IsSimulateModeActive())
+	{
+		if (Key == EKeys::Enter && Event == IE_Released)
+		{
+			ViewModel->StartSimulation();
+			return true;
+		}
+	}
+
+	return FEdMode::InputKey(ViewportClient, Viewport, Key, Event);
 }
 
 bool FContextualAnimEdMode::AllowWidgetMove()
@@ -78,14 +326,70 @@ bool FContextualAnimEdMode::AllowWidgetMove()
 
 bool FContextualAnimEdMode::ShouldDrawWidget() const
 {
-	return SelectedActor.IsValid();
+	return SelectedSelectionCriterionData.IsValid();
+}
+
+bool FContextualAnimEdMode::GetCustomDrawingCoordinateSystem(FMatrix& InMatrix, void* InData)
+{
+	if (ViewModel)
+	{
+		if (SelectedSelectionCriterionData.IsValid())
+		{
+			if (const UContextualAnimSceneInstance* SceneInstance = ViewModel->GetSceneInstance())
+			{
+				FTransform PrimaryActorTransform = FTransform::Identity;
+				if (const FContextualAnimSceneActorData* PreviewData = SceneInstance->FindSceneActorDataByRole(ViewModel->GetSceneAsset()->GetPrimaryRole()))
+				{
+					PrimaryActorTransform = PreviewData->GetActor()->GetTransform();
+					InMatrix = PrimaryActorTransform.ToMatrixNoScale().RemoveTranslation();
+					return true;
+				}
+			}
+		}
+	}
+
+	return false;
+}
+
+bool FContextualAnimEdMode::GetCustomInputCoordinateSystem(FMatrix& InMatrix, void* InData)
+{
+	return GetCustomDrawingCoordinateSystem(InMatrix, InData);
 }
 
 FVector FContextualAnimEdMode::GetWidgetLocation() const
 {
-	if (SelectedActor.IsValid())
+	if (ViewModel)
 	{
-		return SelectedActor->GetActorLocation();
+		if (SelectedSelectionCriterionData.IsValid())
+		{
+			if (const UContextualAnimSceneInstance* SceneInstance = ViewModel->GetSceneInstance())
+			{
+				const UContextualAnimSceneAsset* SceneAsset = ViewModel->GetSceneAsset();
+				if (const FContextualAnimTrack* AnimTrack = SceneAsset->GetAnimTrack(SelectedSelectionCriterionData.RoleName, SelectedSelectionCriterionData.VariantIdx))
+				{
+					if (const UContextualAnimSelectionCriterion_TriggerArea* Spatial = Cast<UContextualAnimSelectionCriterion_TriggerArea>(AnimTrack->SelectionCriteria[SelectedSelectionCriterionData.CriterionIdx]))
+					{
+						FVector Location = FVector::ZeroVector;
+						if (SelectedSelectionCriterionData.DataIdx < 4)
+						{
+							Location = Spatial->PolygonPoints[SelectedSelectionCriterionData.DataIdx];
+						}
+						else
+						{
+							Location = Spatial->PolygonPoints[SelectedSelectionCriterionData.DataIdx - 4] + FVector::UpVector * Spatial->Height;
+						}
+
+						FTransform PrimaryActorTransform = FTransform::Identity;
+						if (const FContextualAnimSceneActorData* PreviewData = SceneInstance->FindSceneActorDataByRole(SceneAsset->GetPrimaryRole()))
+						{
+							PrimaryActorTransform = PreviewData->GetActor()->GetTransform();
+						}
+
+						return PrimaryActorTransform.TransformPositionNoScale(Location);
+					}
+				}
+			}
+		}
 	}
 
 	return FVector::ZeroVector;
diff --git a/Engine/Plugins/Experimental/Animation/ContextualAnimation/Source/ContextualAnimationEditor/Private/ContextualAnimEditorTypes.cpp b/Engine/Plugins/Experimental/Animation/ContextualAnimation/Source/ContextualAnimationEditor/Private/ContextualAnimEditorTypes.cpp
new file mode 100644
--- /dev/null
+++ b/Engine/Plugins/Experimental/Animation/ContextualAnimation/Source/ContextualAnimationEditor/Private/ContextualAnimEditorTypes.cpp
@@ -0,0 +1,75 @@
+// Copyright Epic Games, Inc. All Rights Reserved.
+
+#include "ContextualAnimEditorTypes.h"
+#include "ContextualAnimSceneAsset.h"
+#include "Animation/AnimMontage.h"
+
+// UContextualAnimNewIKTargetParams
+////////////////////////////////////////////////////////////////////////////////////////////////
+
+void UContextualAnimNewIKTargetParams::Reset(const FName& InSourceRole, const UContextualAnimSceneAsset& InSceneAsset)
+{
+	CachedRoles = InSceneAsset.GetRoles();
+	check(CachedRoles.Contains(InSourceRole));
+
+	SceneAssetPtr = &InSceneAsset;
+
+	SourceRole = InSourceRole;
+	GoalName = NAME_None;
+	TargetBone = FBoneReference();
+	SourceBone = FBoneReference();
+
+	if (CachedRoles.Num() > 1)
+	{
+		TargetRole = *CachedRoles.FindByPredicate([this](const FName& Role) { return Role != SourceRole; });
+	}
+}
+
+bool UContextualAnimNewIKTargetParams::HasValidData() const
+{
+	return GoalName != NAME_None &&
+		TargetBone.BoneName != NAME_None &&
+		SourceBone.BoneName != NAME_None &&
+		CachedRoles.Contains(TargetRole);
+}
+
+const UContextualAnimSceneAsset& UContextualAnimNewIKTargetParams::GetSceneAsset() const
+{
+	check(SceneAssetPtr.IsValid());
+	return *SceneAssetPtr.Get();
+}
+
+USkeleton* UContextualAnimNewIKTargetParams::GetSkeleton(bool& bInvalidSkeletonIsError, const IPropertyHandle* PropertyHandle)
+{
+	bInvalidSkeletonIsError = false;
+
+	const FContextualAnimTrack* AnimTrack = nullptr;
+	if (PropertyHandle)
+	{
+		if (PropertyHandle->GetProperty()->GetFName() == GET_MEMBER_NAME_CHECKED(UContextualAnimNewIKTargetParams, SourceBone))
+		{
+			AnimTrack = GetSceneAsset().GetAnimTrack(SourceRole, 0);
+		}
+		else if (PropertyHandle->GetProperty()->GetFName() == GET_MEMBER_NAME_CHECKED(UContextualAnimNewIKTargetParams, TargetBone))
+		{
+			AnimTrack = GetSceneAsset().GetAnimTrack(TargetRole, 0);
+		}
+	}
+
+	return (AnimTrack && AnimTrack->Animation) ? AnimTrack->Animation->GetSkeleton() : nullptr;
+}
+
+TArray<FString> UContextualAnimNewIKTargetParams::GetTargetRoleOptions() const
+{
+	TArray<FString> Options;
+
+	for (const FName& Role : CachedRoles)
+	{
+		if (Role != SourceRole)
+		{
+			Options.Add(Role.ToString());
+		}
+	}
+
+	return Options;
+}
\ No newline at end of file
diff --git a/Engine/Plugins/Experimental/Animation/ContextualAnimation/Source/ContextualAnimationEditor/Private/ContextualAnimMovieSceneNotifyTrackEditor.cpp b/Engine/Plugins/Experimental/Animation/ContextualAnimation/Source/ContextualAnimationEditor/Private/ContextualAnimMovieSceneNotifyTrackEditor.cpp
--- a/Engine/Plugins/Experimental/Animation/ContextualAnimation/Source/ContextualAnimationEditor/Private/ContextualAnimMovieSceneNotifyTrackEditor.cpp
+++ b/Engine/Plugins/Experimental/Animation/ContextualAnimation/Source/ContextualAnimationEditor/Private/ContextualAnimMovieSceneNotifyTrackEditor.cpp
@@ -21,106 +21,24 @@
 #include "Framework/MultiBox/MultiBoxBuilder.h"
 #include "Modules/ModuleManager.h"
 #include "Widgets/Input/SButton.h"
+#include "ContextualAnimEditorTypes.h"
 
 #define LOCTEXT_NAMESPACE "FContextualAnimMovieSceneNotifyTrackEditor"
 
-// FNewRoleWidgetParams
-////////////////////////////////////////////////////////////////////////////////////////////////
-
-bool FNewRoleWidgetParams::HasValidData() const
-{
-	return RoleName != NAME_None && PreviewClass && Animation;
-}
-
-// UNewIKTargetWidgetParams
-////////////////////////////////////////////////////////////////////////////////////////////////
-
-void UNewIKTargetWidgetParams::Reset(const FName& InSourceRole, const UContextualAnimSceneAsset& InSceneAsset)
-{
-	CachedRoles = InSceneAsset.GetRoles();
-	check(CachedRoles.Contains(InSourceRole));
-
-	SceneAssetPtr = &InSceneAsset;
-
-	Params.SourceRole = InSourceRole;
-	Params.GoalName = NAME_None;
-	Params.TargetBone = FBoneReference();
-	Params.SourceBone = FBoneReference();
-
-	if(CachedRoles.Num() > 1)
-	{
-		Params.TargetRole = *CachedRoles.FindByPredicate([this](const FName& Role){ return Role != Params.SourceRole; });
-	}
-}
-
-bool UNewIKTargetWidgetParams::HasValidData() const
-{
-	return	Params.GoalName != NAME_None &&
-			Params.TargetBone.BoneName != NAME_None &&
-			Params.SourceBone.BoneName != NAME_None &&
-			CachedRoles.Contains(Params.TargetRole);
-}
-
-const UContextualAnimSceneAsset& UNewIKTargetWidgetParams::GetSceneAsset() const
-{
-	check(SceneAssetPtr.IsValid());
-	return *SceneAssetPtr.Get();
-}
-
-USkeleton* UNewIKTargetWidgetParams::GetSkeleton(bool& bInvalidSkeletonIsError, const IPropertyHandle* PropertyHandle)
-{
-	//bInvalidSkeletonIsError = false;
-
-	if(PropertyHandle)
-	{
-		if(PropertyHandle->GetProperty()->GetFName() == GET_MEMBER_NAME_CHECKED(FNewIKTargetParams, SourceBone))
-		{
-			UAnimMontage* Animation = GetSceneAsset().GetAnimationForRoleAtIndex(Params.SourceRole, 0);
-			check(Animation);
-
-			return Animation->GetSkeleton();
-		}
-		else if(PropertyHandle->GetProperty()->GetFName() == GET_MEMBER_NAME_CHECKED(FNewIKTargetParams, TargetBone))
-		{
-			UAnimMontage* Animation = GetSceneAsset().GetAnimationForRoleAtIndex(Params.TargetRole, 0);
-			check(Animation);
-
-			return Animation->GetSkeleton();
-		}
-	}
-
-	return nullptr;
-}
-
-TArray<FString> UNewIKTargetWidgetParams::GetTargetRoleOptions() const
-{
-	TArray<FString> Options;
-
-	for(const FName& Role : CachedRoles)
-	{
-		if(Role != Params.SourceRole)
-		{
-			Options.Add(Role.ToString());
-		}
-	}
-	
-	return Options;
-}
-
 // FContextualAnimMovieSceneNotifyTrackEditor
 ////////////////////////////////////////////////////////////////////////////////////////////////
 
 FContextualAnimMovieSceneNotifyTrackEditor::FContextualAnimMovieSceneNotifyTrackEditor(TSharedRef<ISequencer> InSequencer)
 	: FMovieSceneTrackEditor(InSequencer)
-	, NewIKTargetWidgetParams(nullptr)
+	, NewIKTargetParams(nullptr)
 { 
 }
 
 void FContextualAnimMovieSceneNotifyTrackEditor::AddReferencedObjects(FReferenceCollector& Collector)
 {
-	if (NewIKTargetWidgetParams)
+	if (NewIKTargetParams)
 	{
-		Collector.AddReferencedObject(NewIKTargetWidgetParams);
+		Collector.AddReferencedObject(NewIKTargetParams);
 	}
 }
 
@@ -172,67 +90,6 @@ void FContextualAnimMovieSceneNotifyTrackEditor::FillNewNotifyStateMenu(FMenuBui
 void FContextualAnimMovieSceneNotifyTrackEditor::BuildAddTrackMenu(FMenuBuilder& MenuBuilder)
 {
 	// Menu that appears when clicking on the Add Track button next to the Search Tracks bar
-
-	// @TODO: Temporally using this TrackEditor to add New Actor Tracks (new roles to the interaction)
-
-	MenuBuilder.AddSubMenu(
-		LOCTEXT("AddActorTrack", "Actor Track"),
-		LOCTEXT("AddActorTrackTooltip", "Adds new actor track"),
-		FNewMenuDelegate::CreateRaw(this, &FContextualAnimMovieSceneNotifyTrackEditor::BuildAddTrackSubMenu, TArray<FGuid>()),
-		false,
-		FSlateIcon()
-	);
-}
-
-void FContextualAnimMovieSceneNotifyTrackEditor::BuildAddTrackSubMenu(FMenuBuilder& MenuBuilder, TArray<FGuid> ObjectBindings)
-{
-	FPropertyEditorModule& PropertyModule = FModuleManager::LoadModuleChecked<FPropertyEditorModule>("PropertyEditor");
-
-	FDetailsViewArgs Args;
-	Args.bHideSelectionTip = true;
-	Args.bAllowSearch = false;
-	Args.bAllowFavoriteSystem = false;
-
-	NewRoleWidgetParams = MakeShared<FStructOnScope>(FNewRoleWidgetParams::StaticStruct());
-
-	TSharedRef<IStructureDetailsView> StructureDetailsView = PropertyModule.CreateStructureDetailView(Args, FStructureDetailsViewArgs(), NewRoleWidgetParams);
-	
-	MenuBuilder.AddWidget(
-		SNew(SBox)
-		.MinDesiredWidth(500.0f)
-		.MaxDesiredWidth(500.f)
-		.MaxDesiredHeight(400.0f)
-		[
-			SNew(SVerticalBox)
-			+ SVerticalBox::Slot()
-			.AutoHeight()
-			.HAlign(HAlign_Fill)
-			[
-				StructureDetailsView->GetWidget().ToSharedRef()
-			]
-			+SVerticalBox::Slot()
-			.AutoHeight()
-			.HAlign(HAlign_Fill)
-			[
-				SNew(SButton)
-				.ContentPadding(3)
-				.VAlign(VAlign_Center)
-				.HAlign(HAlign_Center)
-				.IsEnabled_Lambda([this]() { return ((FNewRoleWidgetParams*)NewRoleWidgetParams->GetStructMemory())->HasValidData(); })
-				.OnClicked_Lambda([this]()
-					{
-						FNewRoleWidgetParams* Data = (FNewRoleWidgetParams*)NewRoleWidgetParams->GetStructMemory();
-
-						GetMovieSceneSequence().GetViewModel().AddActorTrack(*Data);
-
-						FSlateApplication::Get().DismissAllMenus();
-
-						return FReply::Handled();
-					})
-				.Text(LOCTEXT("OK", "OK"))
-			]
-		],
-		FText(), true, false);
 }
 
 TSharedPtr<SWidget> FContextualAnimMovieSceneNotifyTrackEditor::BuildOutlinerEditWidget(const FGuid& ObjectBinding, UMovieSceneTrack* Track, const FBuildEditWidgetParams& Params)
@@ -311,9 +168,10 @@ void FContextualAnimMovieSceneNotifyTrackEditor::BuildNewIKTargetSubMenu(FMenuBu
 		const FName Role = SceneAsset->FindRoleByAnimation(&Track->GetAnimation());
 		check(Role != NAME_None);
 
-		TArray<FName> ExistingGoals = SceneAsset->GetIKGoalsForRole(Role);
-		for(const FName& GoalName : ExistingGoals)
+		const FContextualAnimIKTargetDefContainer& IKTargets = SceneAsset->GetIKTargetDefsForRole(Role);
+		for (const FContextualAnimIKTargetDefinition& IKTargetDef : IKTargets.IKTargetDefs)
 		{
+			const FName GoalName = IKTargetDef.GoalName;
 			MenuBuilder.AddMenuEntry(
  			FText::FromName(GoalName),
 			FText::GetEmpty(),
@@ -342,18 +200,18 @@ UContextualAnimMovieSceneNotifySection* FContextualAnimMovieSceneNotifyTrackEdit
 
 void FContextualAnimMovieSceneNotifyTrackEditor::BuildNewIKTargetWidget(FMenuBuilder& MenuBuilder, UContextualAnimMovieSceneNotifyTrack* Track, int32 RowIndex)
 {
-	if(NewIKTargetWidgetParams == nullptr)
+	if(NewIKTargetParams == nullptr)
 	{
-		NewIKTargetWidgetParams = NewObject<UNewIKTargetWidgetParams>(UNewIKTargetWidgetParams::StaticClass());
+		NewIKTargetParams = NewObject<UContextualAnimNewIKTargetParams>(UContextualAnimNewIKTargetParams::StaticClass());
 	}
 	
-	check(NewIKTargetWidgetParams);
+	check(NewIKTargetParams);
 
 	const UContextualAnimSceneAsset* SceneAsset = GetMovieSceneSequence().GetViewModel().GetSceneAsset();
 	const FName SourceRole = SceneAsset->FindRoleByAnimation(&Track->GetAnimation());
 	check(SourceRole != NAME_None);
 
-	NewIKTargetWidgetParams->Reset(SourceRole, *SceneAsset);
+	NewIKTargetParams->Reset(SourceRole, *SceneAsset);
 
 	FPropertyEditorModule& PropertyModule = FModuleManager::LoadModuleChecked<FPropertyEditorModule>("PropertyEditor");
 
@@ -363,7 +221,7 @@ void FContextualAnimMovieSceneNotifyTrackEditor::BuildNewIKTargetWidget(FMenuBui
 	Args.bAllowFavoriteSystem = false;
 
 	TSharedPtr<IDetailsView> Widget = PropertyModule.CreateDetailView(Args);
-	Widget->SetObject(NewIKTargetWidgetParams);
+	Widget->SetObject(NewIKTargetParams);
 
 	MenuBuilder.AddWidget(
 		SNew(SBox)
@@ -386,24 +244,30 @@ void FContextualAnimMovieSceneNotifyTrackEditor::BuildNewIKTargetWidget(FMenuBui
 				.ContentPadding(3)
 				.VAlign(VAlign_Center)
 				.HAlign(HAlign_Center)
-				.IsEnabled_Lambda([this](){ return NewIKTargetWidgetParams->HasValidData(); })
+				.IsEnabled_Lambda([this](){ return NewIKTargetParams->HasValidData(); })
 				.OnClicked_Lambda([this, Track, RowIndex]()
 				{
 					UContextualAnimSceneAsset* SceneAsset = GetMovieSceneSequence().GetViewModel().GetSceneAsset();
 					check(SceneAsset);
 
-					FContextualAnimCompositeTrack* SceneAssetTrack = SceneAsset->DataContainer.Find(NewIKTargetWidgetParams->Params.SourceRole);
-					check(SceneAssetTrack);
-
 					// Add IK Target definition to the scene asset
 					FContextualAnimIKTargetDefinition IKTargetDef;
-					IKTargetDef.GoalName = NewIKTargetWidgetParams->Params.GoalName;
-					IKTargetDef.BoneName = NewIKTargetWidgetParams->Params.SourceBone.BoneName;
-					IKTargetDef.Provider = NewIKTargetWidgetParams->Params.Provider;
-					IKTargetDef.TargetRoleName = NewIKTargetWidgetParams->Params.TargetRole;
-					IKTargetDef.TargetBoneName = NewIKTargetWidgetParams->Params.TargetBone.BoneName;
+					IKTargetDef.GoalName = NewIKTargetParams->GoalName;
+					IKTargetDef.BoneName = NewIKTargetParams->SourceBone.BoneName;
+					IKTargetDef.Provider = NewIKTargetParams->Provider;
+					IKTargetDef.TargetRoleName = NewIKTargetParams->TargetRole;
+					IKTargetDef.TargetBoneName = NewIKTargetParams->TargetBone.BoneName;
 
-					SceneAssetTrack->Settings.IKTargetDefinitions.AddUnique(IKTargetDef);
+					if(FContextualAnimIKTargetDefContainer* ContainerPtr = SceneAsset->RoleToIKTargetDefsMap.Find(NewIKTargetParams->SourceRole))
+					{
+						ContainerPtr->IKTargetDefs.AddUnique(IKTargetDef);
+					}
+					else
+					{
+						FContextualAnimIKTargetDefContainer Container;
+						Container.IKTargetDefs.AddUnique(IKTargetDef);
+						SceneAsset->RoleToIKTargetDefsMap.Add(NewIKTargetParams->SourceRole, Container);
+					}
 
 					SceneAsset->PrecomputeData();
 
diff --git a/Engine/Plugins/Experimental/Animation/ContextualAnimation/Source/ContextualAnimationEditor/Private/ContextualAnimPreviewManager.cpp b/Engine/Plugins/Experimental/Animation/ContextualAnimation/Source/ContextualAnimationEditor/Private/ContextualAnimPreviewManager.cpp
deleted file mode 100644
diff --git a/Engine/Plugins/Experimental/Animation/ContextualAnimation/Source/ContextualAnimationEditor/Private/ContextualAnimViewModel.cpp b/Engine/Plugins/Experimental/Animation/ContextualAnimation/Source/ContextualAnimationEditor/Private/ContextualAnimViewModel.cpp
--- a/Engine/Plugins/Experimental/Animation/ContextualAnimation/Source/ContextualAnimationEditor/Private/ContextualAnimViewModel.cpp
+++ b/Engine/Plugins/Experimental/Animation/ContextualAnimation/Source/ContextualAnimationEditor/Private/ContextualAnimViewModel.cpp
@@ -12,20 +12,30 @@
 #include "ContextualAnimMovieSceneNotifyTrack.h"
 #include "ContextualAnimMovieSceneNotifySection.h"
 #include "ContextualAnimPreviewScene.h"
-#include "ContextualAnimPreviewManager.h"
 #include "Animation/AnimMontage.h"
 #include "Animation/AnimInstance.h"
 #include "Animation/AnimNotifies/AnimNotifyState.h"
 #include "EngineUtils.h"
 #include "ContextualAnimTypes.h"
+#include "ContextualAnimEditorTypes.h"
 #include "ContextualAnimUtilities.h"
 #include "ContextualAnimSceneAsset.h"
+#include "ContextualAnimManager.h"
+#include "ContextualAnimSceneInstance.h"
 #include "AnimNotifyState_IKWindow.h"
+#include "PropertyEditorModule.h"
+#include "Modules/ModuleManager.h"
+#include "GameFramework/Character.h"
+#include "GameFramework/CharacterMovementComponent.h"
+#include "Camera/CameraComponent.h"
+#include "Components/CapsuleComponent.h"
+#include "MotionWarpingComponent.h"
 
 FContextualAnimViewModel::FContextualAnimViewModel()
 	: SceneAsset(nullptr)
 	, MovieSceneSequence(nullptr)
 	, MovieScene(nullptr)
+	, ContextualAnimManager(nullptr)
 {
 }
 
@@ -42,9 +52,9 @@ FContextualAnimViewModel::~FContextualAnimViewModel()
 void FContextualAnimViewModel::AddReferencedObjects(FReferenceCollector& Collector)
 {
 	Collector.AddReferencedObject(SceneAsset);
-	Collector.AddReferencedObject(PreviewManager);
 	Collector.AddReferencedObject(MovieSceneSequence);
 	Collector.AddReferencedObject(MovieScene);
+	Collector.AddReferencedObject(ContextualAnimManager);
 }
 
 TSharedPtr<ISequencer> FContextualAnimViewModel::GetSequencer()
@@ -57,17 +67,16 @@ void FContextualAnimViewModel::Initialize(UContextualAnimSceneAsset* InSceneAsse
 	SceneAsset = InSceneAsset;
 	PreviewScenePtr = InPreviewScene;
 
-	PreviewManager = NewObject<UContextualAnimPreviewManager>(GetTransientPackage());
-	PreviewManager->Initialize(*PreviewScenePtr.Pin()->GetWorld(), *SceneAsset);
+	ContextualAnimManager = NewObject<UContextualAnimManager>(GetWorld());
 
 	CreateSequencer();
-	
+
 	RefreshSequencerTracks();
 }
 
 UAnimMontage* FContextualAnimViewModel::FindAnimationByGuid(const FGuid& Guid) const
 {
-	return PreviewManager->FindAnimationByGuid(Guid);
+	return SceneInstance.IsValid() ? const_cast<UAnimMontage*>(SceneInstance->FindSceneActorDataByGuid(Guid)->GetAnimMontage()) : nullptr;
 }
 
 void FContextualAnimViewModel::CreateSequencer()
@@ -98,45 +107,105 @@ void FContextualAnimViewModel::CreateSequencer()
 	Sequencer = SequencerModule.CreateSequencer(SequencerInitParams);
 	Sequencer->OnMovieSceneDataChanged().AddRaw(this, &FContextualAnimViewModel::SequencerDataChanged);
 	Sequencer->OnGlobalTimeChanged().AddRaw(this, &FContextualAnimViewModel::SequencerTimeChanged);
-	//Sequencer->GetSelectionChangedTracks().AddRaw(this, &FContextualAnimAssetEditorToolkit::SequencerTrackSelectionChanged);
-	//Sequencer->GetSelectionChangedSections().AddRaw(this, &FContextualAnimAssetEditorToolkit::SequencerSectionSelectionChanged);
 	Sequencer->SetPlaybackStatus(EMovieScenePlayerStatus::Stopped);
 }
 
-void FContextualAnimViewModel::RefreshSequencerTracks()
+void FContextualAnimViewModel::SetActiveSceneVariantIdx(int32 Index)
+{
+	check(FMath::IsWithinInclusive(Index, 0, (SceneAsset->GetTotalVariants() - 1)));
+
+	ActiveSceneVariantIdx = Index;
+
+	RefreshSequencerTracks();
+}
+
+AActor* FContextualAnimViewModel::SpawnPreviewActor(const FContextualAnimTrack& AnimTrack)
 {
-	// Remove movie scene tracks (if any)
-	for (const FContextualAnimPreviewActorData& Data : PreviewManager->PreviewActorsData)
+	const FContextualAnimRoleDefinition* RoleDef = GetSceneAsset()->RolesAsset ? GetSceneAsset()->RolesAsset->FindRoleDefinitionByName(AnimTrack.Role) : nullptr;
+	UClass* PreviewClass = RoleDef ? RoleDef->PreviewActorClass : nullptr;
+	const FTransform SpawnTransform = (AnimTrack.AlignmentData.ExtractTransformAtTime(0, 0.f));
+
+	FActorSpawnParameters Params;
+	Params.SpawnCollisionHandlingOverride = ESpawnActorCollisionHandlingMethod::AdjustIfPossibleButAlwaysSpawn;
+	AActor* PreviewActor = GetWorld()->SpawnActor<AActor>(PreviewClass, SpawnTransform, Params);
+
+	if (ACharacter* PreviewCharacter = Cast<ACharacter>(PreviewActor))
 	{
-		Data.GetAnimation()->UnregisterOnNotifyChanged(this);
+		PreviewCharacter->bUseControllerRotationYaw = false;
+
+		if (UCharacterMovementComponent* CharacterMovementComp = PreviewCharacter->GetCharacterMovement())
+		{
+			CharacterMovementComp->bOrientRotationToMovement = true;
+			CharacterMovementComp->bUseControllerDesiredRotation = false;
+			CharacterMovementComp->RotationRate = FRotator(0.f, 540.0, 0.f);
+			CharacterMovementComp->bRunPhysicsWithNoController = true;
+
+			CharacterMovementComp->SetMovementMode(AnimTrack.bRequireFlyingMode ? EMovementMode::MOVE_Flying : EMovementMode::MOVE_Walking);
+		}
 
-		MovieSceneSequence->GetMovieScene()->RemovePossessable(Data.Guid);
+		if (UCameraComponent* CameraComp = PreviewCharacter->FindComponentByClass<UCameraComponent>())
+		{
+			CameraComp->DestroyComponent();
+		}
 	}
 
-	// Destroy preview actors (if any)
-	PreviewManager->Reset();
+	UE_LOG(LogContextualAnim, Log, TEXT("Spawned preview Actor: %s at Loc: %s Rot: %s Role: %s"),
+		*GetNameSafe(PreviewActor), *SpawnTransform.GetLocation().ToString(), *SpawnTransform.Rotator().ToString(), *AnimTrack.Role.ToString());
 
-	Sequencer->NotifyMovieSceneDataChanged(EMovieSceneDataChangeType::MovieSceneStructureItemsChanged);
+	return PreviewActor;
+}
 
-	for (const auto& Entry : SceneAsset->DataContainer)
+void FContextualAnimViewModel::RefreshSequencerTracks()
+{
+	// Remove movie scene tracks and destroy existing actors (if any)
+
+	for(UAnimMontage* Anim : AnimsBeingEdited)
+	{
+		Anim->UnregisterOnNotifyChanged(this);
+	}
+	AnimsBeingEdited.Reset();
+
+	if (SceneInstance.IsValid())
+	{
+		SceneInstance->Stop();
+	}
+
+	const int32 Num = MovieSceneSequence->GetMovieScene()->GetPossessableCount();
+	if(Num > 0)
 	{
-		if (Entry.Value.AnimDataContainer.Num() > 0)
+		for (int32 Idx = (Num - 1); Idx >= 0; Idx--)
 		{
-			const FName& Role = Entry.Key;
-			const FContextualAnimData& Data = Entry.Value.AnimDataContainer[0]; //@TODO: Use the first entry for now
+			FMovieScenePossessable& Possessable = MovieSceneSequence->GetMovieScene()->GetPossessable(Idx);
+			MovieSceneSequence->GetMovieScene()->RemovePossessable(Possessable.GetGuid());
+		}
+	}
 
-			UAnimMontage* AnimMontage = Data.Animation;
-			if(AnimMontage == nullptr)
+	if(StartSceneParams.RoleToActorMap.Num() > 0)
+	{
+		for (auto& MapEntry : StartSceneParams.RoleToActorMap)
+		{
+			if (MapEntry.Value)
 			{
-				continue;
+				MapEntry.Value->Destroy();
 			}
+		}
+	}
+
+	Sequencer->NotifyMovieSceneDataChanged(EMovieSceneDataChangeType::MovieSceneStructureItemsChanged);
+
+	StartSceneParams.Reset();
+	StartSceneParams.VariantIdx = ActiveSceneVariantIdx;
+
+	SceneAsset->ForEachAnimTrack(ActiveSceneVariantIdx, [this](const FContextualAnimTrack& AnimTrack)
+		{
+			const FName& Role = AnimTrack.Role;
 
 			// Spawn preview actor
-			AActor* PreviewActor = PreviewManager->SpawnPreviewActor(Role, Data);
+			AActor* PreviewActor = SpawnPreviewActor(AnimTrack);
 			check(PreviewActor);
 
 			// Set actor label so the track shows the name of the role it represents
-			PreviewActor->SetActorLabel(Role.ToString(), false);
+			//PreviewActor->SetActorLabel(Role.ToString(), false);
 
 			// Add preview actors to sequencer
 			const bool bSelectActors = false;
@@ -146,87 +215,205 @@ void FContextualAnimViewModel::RefreshSequencerTracks()
 
 			const FGuid& Guid = Guids[0];
 
-			// Add Animation Track
+			UAnimMontage* AnimMontage = AnimTrack.Animation;
+			if (AnimMontage)
 			{
-				// @TODO: Temporally using an EventTrack to represent the animation since this is just a visual representation of the data. This assumes there is a single section in the montage
-				UMovieSceneEventTrack* AnimTrack = MovieSceneSequence->GetMovieScene()->AddTrack<UMovieSceneEventTrack>(Guid);
-				check(AnimTrack);
+				// Add Animation Track
+				{
+					// @TODO: Temporally using an EventTrack to represent the animation since this is just a visual representation of the data. This assumes there is a single section in the montage
+					UMovieSceneEventTrack* MovieSceneAnimTrack = MovieSceneSequence->GetMovieScene()->AddTrack<UMovieSceneEventTrack>(Guid);
+					check(MovieSceneAnimTrack);
 
-				AnimTrack->SetDisplayName(FText::FromString(GetNameSafe(AnimMontage)));
+					MovieSceneAnimTrack->SetDisplayName(FText::FromString(GetNameSafe(AnimMontage)));
 
-				UMovieSceneSection* NewSection = NewObject<UMovieSceneSection>(AnimTrack, UMovieSceneEventRepeaterSection::StaticClass(), NAME_None, RF_Transactional);
-				check(NewSection);
+					UMovieSceneSection* NewSection = NewObject<UMovieSceneSection>(MovieSceneAnimTrack, UMovieSceneEventRepeaterSection::StaticClass(), NAME_None, RF_Transactional);
+					check(NewSection);
 
-				FFrameRate TickResolution = MovieSceneSequence->GetMovieScene()->GetTickResolution();
-				FFrameNumber StartFrame(0);
-				FFrameNumber EndFrame = (AnimMontage->GetPlayLength() * TickResolution).RoundToFrame();
-				NewSection->SetRange(TRange<FFrameNumber>::Exclusive(StartFrame, EndFrame));
+					FFrameRate TickResolution = MovieSceneSequence->GetMovieScene()->GetTickResolution();
+					FFrameNumber StartFrame(0);
+					FFrameNumber EndFrame = (AnimMontage->GetPlayLength() * TickResolution).RoundToFrame();
+					NewSection->SetRange(TRange<FFrameNumber>::Exclusive(StartFrame, EndFrame));
 
-				AnimTrack->AddSection(*NewSection);
-			}
+					MovieSceneAnimTrack->AddSection(*NewSection);
+				}
 
-			// Add Notify Tracks
-			{
-				for (const FAnimNotifyTrack& NotifyTrack : AnimMontage->AnimNotifyTracks)
+				// Add Notify Tracks
 				{
-					UContextualAnimMovieSceneNotifyTrack* Track = MovieSceneSequence->GetMovieScene()->AddTrack<UContextualAnimMovieSceneNotifyTrack>(Guid);
-					check(Track);
+					for (const FAnimNotifyTrack& NotifyTrack : AnimMontage->AnimNotifyTracks)
+					{
+						UContextualAnimMovieSceneNotifyTrack* Track = MovieSceneSequence->GetMovieScene()->AddTrack<UContextualAnimMovieSceneNotifyTrack>(Guid);
+						check(Track);
 
-					Track->Initialize(*AnimMontage, NotifyTrack);
+						Track->Initialize(*AnimMontage, NotifyTrack);
+					}
+
+					// Listen for when the notifies in the animation changes, so we can refresh the notify sections here
+					AnimMontage->RegisterOnNotifyChanged(UAnimSequenceBase::FOnNotifyChanged::CreateSP(this, &FContextualAnimViewModel::OnAnimNotifyChanged, AnimMontage));
 				}
 
-				// Listen for when the notifies in the animation changes, so we can refresh the notify sections here
-				AnimMontage->RegisterOnNotifyChanged(UAnimSequenceBase::FOnNotifyChanged::CreateSP(this, &FContextualAnimViewModel::OnAnimNotifyChanged, AnimMontage));
+				AnimsBeingEdited.Add(AnimMontage);
 			}
 
-			// Save preview actor data
-			PreviewManager->AddPreviewActor(*PreviewActor, Role, Guid, *AnimMontage);
+			StartSceneParams.RoleToActorMap.Add(Role, PreviewActor);
+
+			return UE::ContextualAnim::EForEachResult::Continue;
+		});
+
+	Sequencer->NotifyMovieSceneDataChanged(EMovieSceneDataChangeType::MovieSceneStructureItemsChanged);
+
+	SceneInstance = ContextualAnimManager->ForceStartScene(*GetSceneAsset(), StartSceneParams);
+
+	// Disable auto blend out
+	if(ensureAlways(SceneInstance.IsValid()))
+	{
+		for (auto& Binding : SceneInstance->GetBindings())
+		{
+			Binding.Guid = Sequencer->FindObjectId(*Binding.GetActor(), MovieSceneSequenceID::Root);
+			check(Binding.Guid.IsValid());
+
+			if (FAnimMontageInstance* MontageInstance = Binding.GetAnimMontageInstance())
+			{
+				MontageInstance->bEnableAutoBlendOut = false;
+			}
 		}
 	}
+}
 
-	//@TODO: DisableCollisionBetweenActors should also be called if bDisableCollisionBetweenActors changes while the editor is open
-	if(SceneAsset->bDisableCollisionBetweenActors)
+void FContextualAnimViewModel::AddNewVariant(const FContextualAnimNewVariantParams& Params)
+{
+	FContextualAnimTracksContainer Container;
+	for (const FContextualAnimNewVariantData& Data : Params.Data)
 	{
-		PreviewManager->DisableCollisionBetweenActors();
+		FContextualAnimTrack AnimTrack;
+		AnimTrack.Role = Data.RoleName;
+		AnimTrack.Animation = Data.Animation;
+		AnimTrack.bRequireFlyingMode = Data.bRequiresFlyingMode;
+		Container.Tracks.Add(AnimTrack);
 	}
 
-	Sequencer->NotifyMovieSceneDataChanged(EMovieSceneDataChangeType::MovieSceneStructureItemsChanged);
-}
+	SceneAsset->Variants.Add(Container);
 
-void FContextualAnimViewModel::AddActorTrack(const FNewRoleWidgetParams& Params)
-{	
-	UE_LOG(LogContextualAnim, Log, TEXT("FContextualAnimViewModel::AddActorTrack Role: %s PreviewClass: %s Animation: %s"), 
-		*Params.RoleName.ToString(), *GetNameSafe(Params.PreviewClass), *GetNameSafe(Params.Animation));
+	SceneAsset->PrecomputeData();
 
-	FContextualAnimCompositeTrack CompositeTrack;
-	CompositeTrack.Settings.PreviewActorClass = Params.PreviewClass;
-	CompositeTrack.Settings.MeshToComponent = Params.MeshToComponent;
+	SceneAsset->MarkPackageDirty();
 
-	FContextualAnimData AnimData;
-	AnimData.Animation = Params.Animation;
-	AnimData.bRequireFlyingMode = Params.bRequiresFlyingMode;
-	CompositeTrack.AnimDataContainer.Add(AnimData);
+	SetActiveSceneVariantIdx(SceneAsset->GetTotalVariants() - 1);
+}
 
-	SceneAsset->DataContainer.Add(Params.RoleName, CompositeTrack);
+void FContextualAnimViewModel::ToggleSimulateMode() 
+{ 
+	bIsSimulateModeActive = !bIsSimulateModeActive; 
 
-	SceneAsset->PrecomputeData();
+	if(bIsSimulateModeActive)
+	{
+		if (SceneInstance.IsValid())
+		{
+			SceneInstance->Stop();
+		}
+	}
+	else
+	{
+		if(SceneInstance.IsValid())
+		{
+			for (auto& Binding : SceneInstance->GetBindings())
+			{
+				if (UMotionWarpingComponent* MotionWarpComp = Binding.GetActor()->FindComponentByClass<UMotionWarpingComponent>())
+				{
+					for (const FContextualAnimAlignmentSectionData AligmentData : SceneAsset->AlignmentSections)
+					{
+						MotionWarpComp->RemoveWarpTarget(AligmentData.WarpTargetName);
+					}
+				}
+			}
+		}
+		
 
-	RefreshSequencerTracks();
+		RefreshSequencerTracks();
+	}
+};
+
+void FContextualAnimViewModel::StartSimulation()
+{
+	//SceneInstance = ContextualAnimManager->ForceStartScene(*GetSceneAsset(), StartSceneParams);
+	SceneInstance = ContextualAnimManager->TryStartScene(*GetSceneAsset(), StartSceneParams);
+
+	if(SceneInstance == nullptr)
+	{
+		//@TODO: This should be a message on the screen
+		UE_LOG(LogContextualAnim, Warning, TEXT("Can't start scene"));
+		return;
+	}
 }
 
-UObject* FContextualAnimViewModel::GetPlaybackContext() const
+UWorld* FContextualAnimViewModel::GetWorld() const
 {
 	check(PreviewScenePtr.IsValid());
 	return PreviewScenePtr.Pin()->GetWorld();
 }
 
+UObject* FContextualAnimViewModel::GetPlaybackContext() const
+{
+	return GetWorld();
+}
+
 void FContextualAnimViewModel::SequencerTimeChanged()
 {
+	auto ResetActorTransform = [](FContextualAnimSceneActorData& SceneActorData, float Time)
+	{
+		const USkeletalMeshComponent* SkelMeshComp = UContextualAnimUtilities::TryGetSkeletalMeshComponent(SceneActorData.GetActor());
+
+		const FTransform RootTransform = UContextualAnimUtilities::ExtractRootTransformFromAnimation(SceneActorData.GetAnimMontage(), Time);
+		const FTransform StartTransform = SkelMeshComp->GetRelativeTransform().Inverse() * RootTransform;
+
+		SceneActorData.GetActor()->SetActorLocationAndRotation(StartTransform.GetLocation(), StartTransform.GetRotation());
+
+		if (UCharacterMovementComponent* MovementComp = SceneActorData.GetActor()->FindComponentByClass<UCharacterMovementComponent>())
+		{
+			MovementComp->StopMovementImmediately();
+		}
+	};
+
 	EMovieScenePlayerStatus::Type CurrentStatus = Sequencer->GetPlaybackStatus();
 	const float CurrentSequencerTime = Sequencer->GetGlobalTime().AsSeconds();
 	const float PlaybackSpeed = Sequencer->GetPlaybackSpeed();
 
-	PreviewManager->PreviewTimeChanged(PreviousSequencerStatus, PreviousSequencerTime, CurrentStatus, CurrentSequencerTime, PlaybackSpeed);
+	if(SceneInstance.IsValid())
+	{
+		for (auto& Binding : SceneInstance->GetBindings())
+		{
+			if (FAnimMontageInstance* MontageInstance = Binding.GetAnimMontageInstance())
+			{
+				const float AnimPlayLength = MontageInstance->Montage->GetPlayLength();
+				float PreviousTime = FMath::Clamp(PreviousSequencerTime, 0.f, AnimPlayLength);
+				float CurrentTime = FMath::Clamp(CurrentSequencerTime, 0.f, AnimPlayLength);
+
+				if (CurrentStatus == EMovieScenePlayerStatus::Stopped || CurrentStatus == EMovieScenePlayerStatus::Scrubbing)
+				{
+					ResetActorTransform(Binding, CurrentTime);
+
+					if (MontageInstance->IsPlaying())
+					{
+						MontageInstance->Pause();
+					}
+
+					MontageInstance->SetPosition(CurrentTime);
+				}
+				else if (CurrentStatus == EMovieScenePlayerStatus::Playing)
+				{
+					if (PlaybackSpeed > 0.f && CurrentTime < PreviousTime)
+					{
+						ResetActorTransform(Binding, CurrentTime);
+						MontageInstance->SetPosition(CurrentTime);
+					}
+
+					if (!MontageInstance->IsPlaying())
+					{
+						MontageInstance->SetPlaying(true);
+					}
+				}
+			}
+		}
+	}
 
 	PreviousSequencerStatus = CurrentStatus;
 	PreviousSequencerTime = CurrentSequencerTime;
@@ -236,14 +423,19 @@ void FContextualAnimViewModel::SequencerDataChanged(EMovieSceneDataChangeType Da
 {
 	UE_LOG(LogContextualAnim, Log, TEXT("FContextualAnimViewModel::OnMovieSceneDataChanged DataChangeType: %d"), (int32)DataChangeType);
 
+	if(!SceneInstance.IsValid())
+	{
+		return;
+	}
+
 	if(DataChangeType == EMovieSceneDataChangeType::TrackValueChanged)
 	{
 		// Update IK AnimNotify's bEnable flag based on the Active state of the section
 		// @TODO: Temp brute-force approach until having a way to override FMovieSceneSection::SetIsActive or something similar
 
-		for (const FContextualAnimPreviewActorData& Data : PreviewManager->PreviewActorsData)
+		for (const auto& Binding : SceneInstance->GetBindings())
 		{
-			TArray<UMovieSceneTrack*> Tracks = MovieSceneSequence->GetMovieScene()->FindTracks(UContextualAnimMovieSceneNotifyTrack::StaticClass(), Data.Guid);
+			TArray<UMovieSceneTrack*> Tracks = MovieSceneSequence->GetMovieScene()->FindTracks(UContextualAnimMovieSceneNotifyTrack::StaticClass(), Binding.Guid);
 			for(UMovieSceneTrack* Track : Tracks)
 			{
 				const TArray<UMovieSceneSection*>& Sections = Track->GetAllSections();
@@ -264,39 +456,6 @@ void FContextualAnimViewModel::SequencerDataChanged(EMovieSceneDataChangeType Da
 			}
 		}
 	}
-	else if(DataChangeType == EMovieSceneDataChangeType::MovieSceneStructureItemRemoved)
-	{
-		// Remove preview actor and role from the scene asset if an actor track was removed
-		int32 TotalRemoved = 0;
-		for (int32 Idx = PreviewManager->PreviewActorsData.Num() - 1; Idx >= 0; Idx--)
-		{
-			const FContextualAnimPreviewActorData& Data = PreviewManager->PreviewActorsData[Idx];
-
-			// Missing binding means the actor track was removed...
-			FMovieSceneBinding* Binding = MovieSceneSequence->GetMovieScene()->FindBinding(Data.Guid);
-			if (Binding == nullptr)
-			{
-				// Remove role from the scene asset
-				SceneAsset->DataContainer.Remove(Data.Role);
-
-				// Remove preview actor from the scene
-				Data.Actor->Destroy();
-
-				// Remove cached data
-				PreviewManager->PreviewActorsData.RemoveAt(Idx, 1, false);
-
-				TotalRemoved++;
-			}
-		}
-
-		PreviewManager->PreviewActorsData.Shrink();
-
-		if(TotalRemoved > 0)
-		{
-			SceneAsset->PrecomputeData();
-			SceneAsset->MarkPackageDirty();
-		}
-	}
 }
 
 void FContextualAnimViewModel::OnAnimNotifyChanged(UAnimMontage* Animation)
@@ -319,4 +478,27 @@ void FContextualAnimViewModel::AnimationModified(UAnimMontage& Animation)
 	Animation.RefreshCacheData();
 	Animation.PostEditChange();
 	Animation.MarkPackageDirty();
+}
+
+void FContextualAnimViewModel::OnPreviewActorClassChanged()
+{
+	const UContextualAnimRolesAsset* RolesAsset = GetSceneAsset()->RolesAsset;
+
+	if(RolesAsset && SceneInstance.IsValid())
+	{
+		for (const auto& Binding : SceneInstance->GetBindings())
+		{
+			if (const FContextualAnimRoleDefinition* RoleDef = RolesAsset->FindRoleDefinitionByName(Binding.GetRole()))
+			{
+				const UClass* DesiredPreviewClass = RoleDef->PreviewActorClass;
+				const UClass* CurrentPreviewClass = Binding.GetActor()->GetClass();
+
+				if (DesiredPreviewClass && DesiredPreviewClass != CurrentPreviewClass)
+				{
+					RefreshSequencerTracks();
+					break;
+				}
+			}
+		}
+	}
 }
\ No newline at end of file
diff --git a/Engine/Plugins/Experimental/Animation/ContextualAnimation/Source/ContextualAnimationEditor/Private/ContextualAnimViewportClient.cpp b/Engine/Plugins/Experimental/Animation/ContextualAnimation/Source/ContextualAnimationEditor/Private/ContextualAnimViewportClient.cpp
--- a/Engine/Plugins/Experimental/Animation/ContextualAnimation/Source/ContextualAnimationEditor/Private/ContextualAnimViewportClient.cpp
+++ b/Engine/Plugins/Experimental/Animation/ContextualAnimation/Source/ContextualAnimationEditor/Private/ContextualAnimViewportClient.cpp
@@ -3,7 +3,6 @@
 #include "ContextualAnimViewportClient.h"
 #include "ContextualAnimAssetEditorToolkit.h"
 #include "ContextualAnimPreviewScene.h"
-#include "ContextualAnimPreviewManager.h"
 #include "SContextualAnimViewport.h"
 #include "AssetEditorModeManager.h"
 #include "UnrealWidget.h"
@@ -38,4 +37,17 @@ void FContextualAnimViewportClient::TrackingStopped()
 void FContextualAnimViewportClient::Draw(const FSceneView* View, FPrimitiveDrawInterface* PDI)
 {
 	FEditorViewportClient::Draw(View, PDI);
+}
+
+void FContextualAnimViewportClient::OnSetIKTargetsDrawMode(EShowIKTargetsDrawMode Mode)
+{
+	if(ShowIKTargetsDrawMode != Mode)
+	{
+		ShowIKTargetsDrawMode = Mode;
+	}
+}
+
+bool FContextualAnimViewportClient::IsIKTargetsDrawModeSet(EShowIKTargetsDrawMode Mode) const
+{
+	return ShowIKTargetsDrawMode == Mode;
 }
\ No newline at end of file
diff --git a/Engine/Plugins/Experimental/Animation/ContextualAnimation/Source/ContextualAnimationEditor/Private/SContextualAnimViewport.cpp b/Engine/Plugins/Experimental/Animation/ContextualAnimation/Source/ContextualAnimationEditor/Private/SContextualAnimViewport.cpp
--- a/Engine/Plugins/Experimental/Animation/ContextualAnimation/Source/ContextualAnimationEditor/Private/SContextualAnimViewport.cpp
+++ b/Engine/Plugins/Experimental/Animation/ContextualAnimation/Source/ContextualAnimationEditor/Private/SContextualAnimViewport.cpp
@@ -3,6 +3,8 @@
 #include "SContextualAnimViewport.h"
 #include "ContextualAnimViewportClient.h"
 #include "Toolkits/AssetEditorToolkit.h"
+#include "SContextualAnimViewportToolbar.h"
+#include "ContextualAnimAssetEditorCommands.h"
 
 void SContextualAnimViewport::Construct(const FArguments& InArgs, const FContextualAnimViewportRequiredArgs& InRequiredArgs)
 {
@@ -16,6 +18,33 @@ void SContextualAnimViewport::Construct(const FArguments& InArgs, const FContext
 	);
 }
 
+void SContextualAnimViewport::BindCommands()
+{
+	SEditorViewport::BindCommands();
+
+	const FContextualAnimAssetEditorCommands& Commands = FContextualAnimAssetEditorCommands::Get();
+
+	TSharedRef<FContextualAnimViewportClient> ViewportClientRef = ViewportClient.ToSharedRef();
+
+	CommandList->MapAction(
+		Commands.ShowIKTargetsDrawAll,
+		FExecuteAction::CreateSP(ViewportClientRef, &FContextualAnimViewportClient::OnSetIKTargetsDrawMode, EShowIKTargetsDrawMode::All),
+		FCanExecuteAction(),
+		FIsActionChecked::CreateSP(ViewportClientRef, &FContextualAnimViewportClient::IsIKTargetsDrawModeSet, EShowIKTargetsDrawMode::All));
+
+	CommandList->MapAction(
+		Commands.ShowIKTargetsDrawSelected,
+		FExecuteAction::CreateSP(ViewportClientRef, &FContextualAnimViewportClient::OnSetIKTargetsDrawMode, EShowIKTargetsDrawMode::Selected),
+		FCanExecuteAction(),
+		FIsActionChecked::CreateSP(ViewportClientRef, &FContextualAnimViewportClient::IsIKTargetsDrawModeSet, EShowIKTargetsDrawMode::Selected));
+
+	CommandList->MapAction(
+		Commands.ShowIKTargetsDrawNone,
+		FExecuteAction::CreateSP(ViewportClientRef, &FContextualAnimViewportClient::OnSetIKTargetsDrawMode, EShowIKTargetsDrawMode::None),
+		FCanExecuteAction(),
+		FIsActionChecked::CreateSP(ViewportClientRef, &FContextualAnimViewportClient::IsIKTargetsDrawModeSet, EShowIKTargetsDrawMode::None));
+}
+
 TSharedRef<FEditorViewportClient> SContextualAnimViewport::MakeEditorViewportClient()
 {
 	ViewportClient = MakeShared<FContextualAnimViewportClient>(PreviewScenePtr.Pin().ToSharedRef(), SharedThis(this), AssetEditorToolkitPtr.Pin().ToSharedRef());
@@ -25,4 +54,24 @@ TSharedRef<FEditorViewportClient> SContextualAnimViewport::MakeEditorViewportCli
 	ViewportClient->SetViewRotation(EditorViewportDefs::DefaultPerspectiveViewRotation);
 
 	return ViewportClient.ToSharedRef();
+}
+
+TSharedPtr<SWidget> SContextualAnimViewport::MakeViewportToolbar()
+{
+	return SAssignNew(ViewportToolbar, SContextualAnimViewportToolBar, SharedThis(this));
+}
+
+TSharedRef<SEditorViewport> SContextualAnimViewport::GetViewportWidget()
+{
+	return SharedThis(this);
+}
+
+TSharedPtr<FExtender> SContextualAnimViewport::GetExtenders() const
+{
+	TSharedPtr<FExtender> Result(MakeShareable(new FExtender));
+	return Result;
+}
+
+void SContextualAnimViewport::OnFloatingButtonClicked()
+{
 }
\ No newline at end of file
diff --git a/Engine/Plugins/Experimental/Animation/ContextualAnimation/Source/ContextualAnimationEditor/Private/SContextualAnimViewportToolbar.cpp b/Engine/Plugins/Experimental/Animation/ContextualAnimation/Source/ContextualAnimationEditor/Private/SContextualAnimViewportToolbar.cpp
new file mode 100644
--- /dev/null
+++ b/Engine/Plugins/Experimental/Animation/ContextualAnimation/Source/ContextualAnimationEditor/Private/SContextualAnimViewportToolbar.cpp
@@ -0,0 +1,45 @@
+// Copyright Epic Games, Inc. All Rights Reserved.
+
+#include "SContextualAnimViewportToolbar.h"
+#include "SContextualAnimViewport.h"
+#include "PreviewProfileController.h"
+#include "ContextualAnimAssetEditorCommands.h"
+
+#define LOCTEXT_NAMESPACE "ContextualAnimViewportToolBar"
+
+void SContextualAnimViewportToolBar::Construct(const FArguments& InArgs, TSharedPtr<SContextualAnimViewport> InViewport)
+{
+	SCommonEditorViewportToolbarBase::Construct(SCommonEditorViewportToolbarBase::FArguments().AddRealtimeButton(false).PreviewProfileController(MakeShared<FPreviewProfileController>()), InViewport);
+}
+
+TSharedRef<SWidget> SContextualAnimViewportToolBar::GenerateShowMenu() const
+{
+	GetInfoProvider().OnFloatingButtonClicked();
+
+	TSharedRef<SEditorViewport> ViewportRef = GetInfoProvider().GetViewportWidget();
+
+	const bool bInShouldCloseWindowAfterMenuSelection = true;
+	FMenuBuilder ShowMenuBuilder(bInShouldCloseWindowAfterMenuSelection, ViewportRef->GetCommandList());
+	{
+		ShowMenuBuilder.AddSubMenu(
+			LOCTEXT("ShowMenu_IKTargetsDrawSubMenu", "IK Targets"),
+			LOCTEXT("ShowMenu_IKTargetsDrawSubMenuToolTip", "IK Targets Drawing Options"),
+			FNewMenuDelegate::CreateLambda([](FMenuBuilder& SubMenuBuilder)
+				{
+					const FContextualAnimAssetEditorCommands& Commands = FContextualAnimAssetEditorCommands::Get();
+
+					SubMenuBuilder.BeginSection("IKTargets", LOCTEXT("ShowMenu_IKTargetsLabel", "IK Targets"));
+					{
+						SubMenuBuilder.AddMenuEntry(Commands.ShowIKTargetsDrawSelected);
+						SubMenuBuilder.AddMenuEntry(Commands.ShowIKTargetsDrawAll);
+						SubMenuBuilder.AddMenuEntry(Commands.ShowIKTargetsDrawNone);
+					}
+					SubMenuBuilder.EndSection();
+				})
+		);
+	}
+
+	return ShowMenuBuilder.MakeWidget();
+}
+
+#undef LOCTEXT_NAMESPACE
\ No newline at end of file
diff --git a/Engine/Plugins/Experimental/Animation/ContextualAnimation/Source/ContextualAnimationEditor/Public/ContextualAnimAssetEditorCommands.h b/Engine/Plugins/Experimental/Animation/ContextualAnimation/Source/ContextualAnimationEditor/Public/ContextualAnimAssetEditorCommands.h
--- a/Engine/Plugins/Experimental/Animation/ContextualAnimation/Source/ContextualAnimationEditor/Public/ContextualAnimAssetEditorCommands.h
+++ b/Engine/Plugins/Experimental/Animation/ContextualAnimation/Source/ContextualAnimationEditor/Public/ContextualAnimAssetEditorCommands.h
@@ -19,4 +19,12 @@ public:
 public:
 
 	TSharedPtr<FUICommandInfo> ResetPreviewScene;
+
+	TSharedPtr<FUICommandInfo> ShowIKTargetsDrawSelected;
+
+	TSharedPtr<FUICommandInfo> ShowIKTargetsDrawAll;
+
+	TSharedPtr<FUICommandInfo> ShowIKTargetsDrawNone;
+
+	TSharedPtr<FUICommandInfo> Simulate;
 };
diff --git a/Engine/Plugins/Experimental/Animation/ContextualAnimation/Source/ContextualAnimationEditor/Public/ContextualAnimAssetEditorToolkit.h b/Engine/Plugins/Experimental/Animation/ContextualAnimation/Source/ContextualAnimationEditor/Public/ContextualAnimAssetEditorToolkit.h
--- a/Engine/Plugins/Experimental/Animation/ContextualAnimation/Source/ContextualAnimationEditor/Public/ContextualAnimAssetEditorToolkit.h
+++ b/Engine/Plugins/Experimental/Animation/ContextualAnimation/Source/ContextualAnimationEditor/Public/ContextualAnimAssetEditorToolkit.h
@@ -33,9 +33,12 @@ public:
 	virtual FString GetWorldCentricTabPrefix() const override;
 
 	UContextualAnimSceneAsset* GetSceneAsset() const;
-	UContextualAnimPreviewManager* GetPreviewManager() const;
+	FContextualAnimViewModel* GetViewModel() const { return ViewModel.Get(); }
 
 	void ResetPreviewScene();
+	void ToggleSimulateMode();
+	bool IsSimulateModeActive() const;
+	
 
 private:
 	
@@ -49,6 +52,8 @@ private:
 	void ExtendToolbar();
 	void FillToolbar(FToolBarBuilder& ToolbarBuilder);
 
+	TSharedRef<SWidget> BuildVariantsMenu();
+
 	void OnFinishedChangingProperties(const FPropertyChangedEvent& PropertyChangedEvent);
 
 	TSharedPtr<SContextualAnimViewport> ViewportWidget;
@@ -60,6 +65,8 @@ private:
 	TSharedPtr<FContextualAnimPreviewScene> PreviewScene;
 
 	TSharedPtr<FContextualAnimViewModel> ViewModel;
+
+	TSharedPtr<FStructOnScope> NewVariantWidgetStruct;
 };
 
 
diff --git a/Engine/Plugins/Experimental/Animation/ContextualAnimation/Source/ContextualAnimationEditor/Public/ContextualAnimEdMode.h b/Engine/Plugins/Experimental/Animation/ContextualAnimation/Source/ContextualAnimationEditor/Public/ContextualAnimEdMode.h
--- a/Engine/Plugins/Experimental/Animation/ContextualAnimation/Source/ContextualAnimationEditor/Public/ContextualAnimEdMode.h
+++ b/Engine/Plugins/Experimental/Animation/ContextualAnimation/Source/ContextualAnimationEditor/Public/ContextualAnimEdMode.h
@@ -5,6 +5,54 @@
 #include "CoreMinimal.h"
 #include "EdMode.h"
 
+class AActor;
+class FContextualAnimViewModel;
+
+struct FSelectionCriterionHitProxyData
+{
+	FName RoleName = NAME_None;
+	int32 VariantIdx = INDEX_NONE;
+	int32 CriterionIdx = INDEX_NONE;
+	int32 DataIdx = INDEX_NONE;
+
+	FSelectionCriterionHitProxyData(){}
+	FSelectionCriterionHitProxyData(const FName& InRoleName, int32 InVariantIdx, int32 InCriterionIdx, int32 InDataIdx)
+		: RoleName(InRoleName)
+		, VariantIdx(InVariantIdx)
+		, CriterionIdx(InCriterionIdx)
+		, DataIdx(InDataIdx)
+	{}
+
+	void Reset()
+	{
+		RoleName = NAME_None;
+		VariantIdx = INDEX_NONE;
+		CriterionIdx = INDEX_NONE;
+		DataIdx = INDEX_NONE;
+	}
+
+	bool IsValid() const 
+	{ 
+		return	RoleName != NAME_None && 
+				VariantIdx != INDEX_NONE && 
+				CriterionIdx != INDEX_NONE; 
+	}
+};
+
+struct HSelectionCriterionHitProxy : public HHitProxy
+{
+	DECLARE_HIT_PROXY();
+
+	HSelectionCriterionHitProxy(const FSelectionCriterionHitProxyData& InData, EHitProxyPriority InPriority = HPP_Wireframe)
+		: HHitProxy(InPriority)
+		, Data(InData)
+	{}
+
+	virtual EMouseCursor::Type GetMouseCursor() override { return EMouseCursor::Crosshairs; }
+
+	FSelectionCriterionHitProxyData Data;
+};
+
 class FContextualAnimEdMode : public FEdMode
 {
 public:
@@ -14,17 +62,28 @@ public:
 	FContextualAnimEdMode();
 	virtual ~FContextualAnimEdMode();
 
+	virtual void Tick(FEditorViewportClient* ViewportClient, float DeltaTime) override;
+	virtual void Render(const FSceneView* View, FViewport* Viewport, FPrimitiveDrawInterface* PDI);
 	virtual bool HandleClick(FEditorViewportClient* InViewportClient, HHitProxy* HitProxy, const FViewportClick& Click) override;
 	virtual bool StartTracking(FEditorViewportClient* InViewportClient, FViewport* InViewport) override;
 	virtual bool EndTracking(FEditorViewportClient* InViewportClient, FViewport* InViewport) override;
 	virtual bool InputDelta(FEditorViewportClient* InViewportClient, FViewport* InViewport, FVector& InDrag, FRotator& InRot, FVector& InScale) override;
+	virtual bool InputKey(FEditorViewportClient* ViewportClient, FViewport* Viewport, FKey Key, EInputEvent Event) override;
 	virtual bool AllowWidgetMove() override;
 	virtual bool ShouldDrawWidget() const override;
 	virtual FVector GetWidgetLocation() const override;
+	virtual bool GetCustomDrawingCoordinateSystem(FMatrix& InMatrix, void* InData);
+	virtual bool GetCustomInputCoordinateSystem(FMatrix& InMatrix, void* InData);
 
 	bool GetHitResultUnderCursor(FHitResult& OutHitResult, FEditorViewportClient* InViewportClient, const FViewportClick& Click) const;
 
+	void DrawIKTargetsForSceneActor(class FPrimitiveDrawInterface& PDI, const struct FContextualAnimSceneActorData& SceneActorData) const;
+
 private:
 
-	TWeakObjectPtr<class AActor> SelectedActor;
+	FContextualAnimViewModel* ViewModel = nullptr;
+
+	TWeakObjectPtr<AActor> SelectedActor;
+
+	FSelectionCriterionHitProxyData SelectedSelectionCriterionData;
 };
diff --git a/Engine/Plugins/Experimental/Animation/ContextualAnimation/Source/ContextualAnimationEditor/Public/ContextualAnimEditorTypes.h b/Engine/Plugins/Experimental/Animation/ContextualAnimation/Source/ContextualAnimationEditor/Public/ContextualAnimEditorTypes.h
new file mode 100644
--- /dev/null
+++ b/Engine/Plugins/Experimental/Animation/ContextualAnimation/Source/ContextualAnimationEditor/Public/ContextualAnimEditorTypes.h
@@ -0,0 +1,84 @@
+// Copyright Epic Games, Inc. All Rights Reserved.
+
+#pragma once
+
+#include "CoreMinimal.h"
+#include "ContextualAnimTypes.h"
+#include "Interfaces/Interface_BoneReferenceSkeletonProvider.h"
+#include "ContextualAnimEditorTypes.generated.h"
+
+class UContextualAnimSceneAsset;
+
+/** Object used to construct the New IK Target Widget */
+UCLASS()
+class UContextualAnimNewIKTargetParams : public UObject, public IBoneReferenceSkeletonProvider
+{
+	GENERATED_BODY()
+
+public:
+
+	UPROPERTY(VisibleAnywhere, Category = "Settings")
+	FName SourceRole = NAME_None;
+
+	UPROPERTY(EditAnywhere, Category = "Settings")
+	FBoneReference SourceBone;
+
+	UPROPERTY(EditAnywhere, Category = "Settings")
+	EContextualAnimIKTargetProvider Provider = EContextualAnimIKTargetProvider::Autogenerated;
+
+	UPROPERTY(EditAnywhere, Category = "Settings", meta = (GetOptions = "GetTargetRoleOptions"))
+	FName TargetRole = NAME_None;
+
+	UPROPERTY(EditAnywhere, Category = "Settings")
+	FBoneReference TargetBone;
+
+	UPROPERTY(EditAnywhere, Category = "Settings")
+	FName GoalName = NAME_None;
+
+	UContextualAnimNewIKTargetParams(const FObjectInitializer& ObjectInitializer) : Super(ObjectInitializer) {}
+
+	// ~IBoneReferenceSkeletonProvider Interface
+	virtual USkeleton* GetSkeleton(bool& bInvalidSkeletonIsError, const IPropertyHandle* PropertyHandle = nullptr) override;
+
+	void Reset(const FName& InSourceRole, const UContextualAnimSceneAsset& InSceneAsset);
+
+	bool HasValidData() const;
+
+	const UContextualAnimSceneAsset& GetSceneAsset() const;
+
+	UFUNCTION()
+	TArray<FString> GetTargetRoleOptions() const;
+
+private:
+
+	UPROPERTY()
+	TWeakObjectPtr<const UContextualAnimSceneAsset> SceneAssetPtr;
+
+	UPROPERTY()
+	TArray<FName> CachedRoles;
+};
+
+USTRUCT()
+struct FContextualAnimNewVariantData
+{
+	GENERATED_BODY()
+
+	UPROPERTY(VisibleAnywhere, Category = "Settings")
+	FName RoleName = NAME_None;
+
+	UPROPERTY(EditAnywhere, Category = "Settings")
+	UAnimMontage* Animation = nullptr;
+
+	UPROPERTY(EditAnywhere, Category = "Settings")
+	bool bRequiresFlyingMode = false;
+};
+
+/** Struct used to construct the widget for adding a new variant */
+USTRUCT()
+struct FContextualAnimNewVariantParams
+{
+	GENERATED_BODY()
+
+	UPROPERTY(EditAnywhere, EditFixedSize, Category = "Settings", meta = (TitleProperty = "RoleName"))
+	TArray<FContextualAnimNewVariantData> Data;
+};
\ No newline at end of file
diff --git a/Engine/Plugins/Experimental/Animation/ContextualAnimation/Source/ContextualAnimationEditor/Public/ContextualAnimMovieSceneNotifyTrackEditor.h b/Engine/Plugins/Experimental/Animation/ContextualAnimation/Source/ContextualAnimationEditor/Public/ContextualAnimMovieSceneNotifyTrackEditor.h
--- a/Engine/Plugins/Experimental/Animation/ContextualAnimation/Source/ContextualAnimationEditor/Public/ContextualAnimMovieSceneNotifyTrackEditor.h
+++ b/Engine/Plugins/Experimental/Animation/ContextualAnimation/Source/ContextualAnimationEditor/Public/ContextualAnimMovieSceneNotifyTrackEditor.h
@@ -6,8 +6,6 @@
 #include "MovieSceneTrackEditor.h"
 #include "ContextualAnimTypes.h"
 #include "BoneContainer.h"
-#include "Interfaces/Interface_BoneReferenceSkeletonProvider.h"
-#include "ContextualAnimMovieSceneNotifyTrackEditor.generated.h"
 
 class AActor;
 class UAnimMontage;
@@ -16,89 +14,7 @@ class UContextualAnimSceneAsset;
 class UContextualAnimMovieSceneNotifySection;
 class UContextualAnimMovieSceneNotifyTrack;
 class UContextualAnimMovieSceneSequence;
-
-/** Struct used to construct the New Role Widget */
-USTRUCT()
-struct FNewRoleWidgetParams
-{
-	GENERATED_BODY()
-
-	UPROPERTY(EditAnywhere, Category = "Settings")
-	FName RoleName = NAME_None;
-
-	UPROPERTY(EditAnywhere, Category = "Settings")
-	TSubclassOf<AActor> PreviewClass;
-
-	UPROPERTY(EditAnywhere, Category = "Settings")
-	UAnimMontage* Animation = nullptr;
-
-	UPROPERTY(EditAnywhere, Category = "Settings")
-	bool bRequiresFlyingMode = false;
-
-	UPROPERTY(EditAnywhere, Category = "Settings")
-	FTransform MeshToComponent = FTransform(FRotator(0.f, -90.f, 0.f));
-
-	bool HasValidData() const;
-};
-
-USTRUCT()
-struct FNewIKTargetParams
-{
-	GENERATED_BODY()
-
-	UPROPERTY(VisibleAnywhere, Category = "Settings")
-	FName SourceRole = NAME_None;
-
-	UPROPERTY(EditAnywhere, Category = "Settings")
-	FBoneReference SourceBone;
-
-	UPROPERTY(EditAnywhere, Category = "Settings")
-	EContextualAnimIKTargetProvider Provider = EContextualAnimIKTargetProvider::Autogenerated;
-
-	UPROPERTY(EditAnywhere, Category = "Settings", meta = (GetOptions = "GetTargetRoleOptions"))
-	FName TargetRole = NAME_None;
-
-	UPROPERTY(EditAnywhere, Category = "Settings")
-	FBoneReference TargetBone;
-
-	UPROPERTY(EditAnywhere, Category = "Settings")
-	FName GoalName = NAME_None;
-
-};
-
-/** Object used to construct the New IK Target Widget */
-UCLASS()
-class UNewIKTargetWidgetParams : public UObject, public IBoneReferenceSkeletonProvider
-{
-	GENERATED_BODY()
-
-public:
-
-	UPROPERTY(EditAnywhere, Category = "Settings", meta = (ShowOnlyInnerProperties))
-	FNewIKTargetParams Params;
-
-	UNewIKTargetWidgetParams(const FObjectInitializer& ObjectInitializer) : Super(ObjectInitializer) {}
-
-	// ~IBoneReferenceSkeletonProvider Interface
-	virtual USkeleton* GetSkeleton(bool& bInvalidSkeletonIsError, const IPropertyHandle* PropertyHandle = nullptr) override;
-
-	void Reset(const FName& InSourceRole, const UContextualAnimSceneAsset& InSceneAsset);
-
-	bool HasValidData() const;
-
-	const UContextualAnimSceneAsset& GetSceneAsset() const;
-
-	UFUNCTION()
-	TArray<FString> GetTargetRoleOptions() const;
-
-private:
-
-	UPROPERTY()
-	TWeakObjectPtr<const UContextualAnimSceneAsset> SceneAssetPtr;
-
-	UPROPERTY()
-	TArray<FName> CachedRoles;
-};
+class UContextualAnimNewIKTargetParams;
 
 /** Handles section drawing and manipulation of a MovieSceneNotifyTrack */
 class FContextualAnimMovieSceneNotifyTrackEditor : public FMovieSceneTrackEditor, public FGCObject
@@ -125,11 +41,7 @@ public:
 
 private:
 
-	TObjectPtr<UNewIKTargetWidgetParams> NewIKTargetWidgetParams;
-
-	TSharedPtr<FStructOnScope> NewRoleWidgetParams;
-
-	void BuildAddTrackSubMenu(FMenuBuilder& MenuBuilder, TArray<FGuid>);
+	TObjectPtr<UContextualAnimNewIKTargetParams> NewIKTargetParams;
 
 	UContextualAnimMovieSceneNotifySection* CreateNewSection(UContextualAnimMovieSceneNotifyTrack* Track, int32 RowIndex, UClass* NotifyClass);
 
diff --git a/Engine/Plugins/Experimental/Animation/ContextualAnimation/Source/ContextualAnimationEditor/Public/ContextualAnimPreviewManager.h b/Engine/Plugins/Experimental/Animation/ContextualAnimation/Source/ContextualAnimationEditor/Public/ContextualAnimPreviewManager.h
deleted file mode 100644
diff --git a/Engine/Plugins/Experimental/Animation/ContextualAnimation/Source/ContextualAnimationEditor/Public/ContextualAnimViewModel.h b/Engine/Plugins/Experimental/Animation/ContextualAnimation/Source/ContextualAnimationEditor/Public/ContextualAnimViewModel.h
--- a/Engine/Plugins/Experimental/Animation/ContextualAnimation/Source/ContextualAnimationEditor/Public/ContextualAnimViewModel.h
+++ b/Engine/Plugins/Experimental/Animation/ContextualAnimation/Source/ContextualAnimationEditor/Public/ContextualAnimViewModel.h
@@ -9,16 +9,23 @@
 #include "ContextualAnimMovieSceneSequence.h"
 #include "ISequencer.h"
 
+class UWorld;
 class FContextualAnimPreviewScene;
 class UContextualAnimSceneAsset;
-class UContextualAnimPreviewManager;
+class UContextualAnimManager;
 class UContextualAnimMovieSceneSequence;
+class UContextualAnimSceneInstance;
 class UMovieScene;
+class UMovieSceneTrack;
+class UMovieSceneSection;
 class UContextualAnimMovieSceneNotifyTrack;
 class UContextualAnimMovieSceneNotifySection;
 class UAnimMontage;
+class IDetailsView;
+class IStructureDetailsView;
 struct FMovieSceneSectionMovedParams;
-struct FNewRoleWidgetParams;
+struct FContextualAnimNewVariantParams;
+struct FContextualAnimTrack;
 
 class FContextualAnimViewModel : public TSharedFromThis<FContextualAnimViewModel>, public FGCObject
 {
@@ -40,22 +47,29 @@ public:
 	UMovieScene* GetMovieScene() const { return MovieScene; }
 	UContextualAnimMovieSceneSequence* GetMovieSceneSequence() const { return MovieSceneSequence; }
 	UContextualAnimSceneAsset* GetSceneAsset() const { return SceneAsset; }
-	UContextualAnimPreviewManager* GetPreviewManager() const { return PreviewManager; }
+	UContextualAnimSceneInstance* GetSceneInstance() const { return SceneInstance.Get(); }
 
-	void AddActorTrack(const FNewRoleWidgetParams& Params);
+	void AddNewVariant(const FContextualAnimNewVariantParams& Params);
 
  	UAnimMontage* FindAnimationByGuid(const FGuid& Guid) const;
 
 	void AnimationModified(UAnimMontage& Animation);
 
+	void SetActiveSceneVariantIdx(int32 Index);
+
+	int32 GetActiveSceneVariantIdx() const { return ActiveSceneVariantIdx; }
+
+	void OnPreviewActorClassChanged();
+
+	void ToggleSimulateMode();
+	bool IsSimulateModeActive() { return bIsSimulateModeActive; }
+	void StartSimulation();
+
 private:
 
 	/** Scene asset being viewed and edited by this view model. */
 	TObjectPtr<UContextualAnimSceneAsset> SceneAsset;
 
-	/** Manager to spawn and interact with preview actors */
-	TObjectPtr<UContextualAnimPreviewManager> PreviewManager;
-
 	/** MovieSceneSequence for displaying this scene asset in the sequencer time line. */
 	TObjectPtr<UContextualAnimMovieSceneSequence> MovieSceneSequence;
 
@@ -68,6 +82,12 @@ private:
 	/** Weak pointer to the PreviewScene */
 	TWeakPtr<FContextualAnimPreviewScene> PreviewScenePtr;
 
+	TObjectPtr<UContextualAnimManager> ContextualAnimManager;
+
+	TWeakObjectPtr<UContextualAnimSceneInstance> SceneInstance;
+
+	int32 ActiveSceneVariantIdx = 0;
+
 	/** The previous play status for sequencer time line. */
 	EMovieScenePlayerStatus::Type PreviousSequencerStatus;
 
@@ -77,6 +97,17 @@ private:
 	/** Flag for preventing OnAnimNotifyChanged from updating tracks when the change to the animation came from us */
 	bool bUpdatingAnimationFromSequencer = false;
 
+	bool bIsSimulateModeActive = false;
+
+	FContextualAnimStartSceneParams StartSceneParams;
+
+	/** Container for the animations on the time line. Should be removed once we add a proper animation track */
+	TArray<UAnimMontage*> AnimsBeingEdited;
+
+	AActor* SpawnPreviewActor(const FContextualAnimTrack& AnimTrack);
+
+	UWorld* GetWorld() const;
+
 	UObject* GetPlaybackContext() const;
 
 	void SequencerTimeChanged();
diff --git a/Engine/Plugins/Experimental/Animation/ContextualAnimation/Source/ContextualAnimationEditor/Public/ContextualAnimViewportClient.h b/Engine/Plugins/Experimental/Animation/ContextualAnimation/Source/ContextualAnimationEditor/Public/ContextualAnimViewportClient.h
--- a/Engine/Plugins/Experimental/Animation/ContextualAnimation/Source/ContextualAnimationEditor/Public/ContextualAnimViewportClient.h
+++ b/Engine/Plugins/Experimental/Animation/ContextualAnimation/Source/ContextualAnimationEditor/Public/ContextualAnimViewportClient.h
@@ -9,6 +9,13 @@ class FContextualAnimPreviewScene;
 class SContextualAnimViewport;
 class FContextualAnimAssetEditorToolkit;
 
+enum class EShowIKTargetsDrawMode : uint8
+{
+	None,
+	Selected,
+	All
+};
+
 class FContextualAnimViewportClient : public FEditorViewportClient
 {
 public:
@@ -27,6 +34,10 @@ public:
 
 	TSharedRef<FContextualAnimAssetEditorToolkit> GetAssetEditorToolkit() const { return AssetEditorToolkitPtr.Pin().ToSharedRef(); }
 
+	void OnSetIKTargetsDrawMode(EShowIKTargetsDrawMode Mode);
+	bool IsIKTargetsDrawModeSet(EShowIKTargetsDrawMode Mode) const;
+	EShowIKTargetsDrawMode GetShowIKTargetsDrawMode() const { return ShowIKTargetsDrawMode; }
+
 private:
 
 	/** Preview scene we are viewing */
@@ -34,4 +45,6 @@ private:
 
 	/** Asset editor toolkit we are embedded in */
 	TWeakPtr<FContextualAnimAssetEditorToolkit> AssetEditorToolkitPtr;
+
+	EShowIKTargetsDrawMode ShowIKTargetsDrawMode = EShowIKTargetsDrawMode::None;
 };
diff --git a/Engine/Plugins/Experimental/Animation/ContextualAnimation/Source/ContextualAnimationEditor/Public/SContextualAnimViewport.h b/Engine/Plugins/Experimental/Animation/ContextualAnimation/Source/ContextualAnimationEditor/Public/SContextualAnimViewport.h
--- a/Engine/Plugins/Experimental/Animation/ContextualAnimation/Source/ContextualAnimationEditor/Public/SContextualAnimViewport.h
+++ b/Engine/Plugins/Experimental/Animation/ContextualAnimation/Source/ContextualAnimationEditor/Public/SContextualAnimViewport.h
@@ -4,10 +4,12 @@
 
 #include "CoreMinimal.h"
 #include "SEditorViewport.h"
+#include "SCommonEditorViewportToolbarBase.h"
 
 class FContextualAnimPreviewScene;
 class FContextualAnimViewportClient;
 class FContextualAnimAssetEditorToolkit;
+class SContextualAnimViewportToolBar;
 
 struct FContextualAnimViewportRequiredArgs
 {
@@ -21,7 +23,7 @@ struct FContextualAnimViewportRequiredArgs
 	TSharedRef<FContextualAnimPreviewScene> PreviewScene;
 };
 
-class SContextualAnimViewport : public SEditorViewport
+class SContextualAnimViewport : public SEditorViewport, public ICommonEditorViewportToolbarInfoProvider
 {
 public:
 	
@@ -31,14 +33,23 @@ public:
 	void Construct(const FArguments& InArgs, const FContextualAnimViewportRequiredArgs& InRequiredArgs);
 	virtual ~SContextualAnimViewport(){}
 
+	// ~ICommonEditorViewportToolbarInfoProvider interface
+	virtual TSharedRef<class SEditorViewport> GetViewportWidget() override;
+	virtual TSharedPtr<FExtender> GetExtenders() const override;
+	virtual void OnFloatingButtonClicked() override;
+	// ~End of ICommonEditorViewportToolbarInfoProvider interface
+
 protected:
 
 	// ~SEditorViewport interface
+	virtual void BindCommands() override;
 	virtual TSharedRef<FEditorViewportClient> MakeEditorViewportClient() override;
+	virtual TSharedPtr<SWidget> MakeViewportToolbar() override;
 	// ~End of SEditorViewport interface
 
-protected:
-	
+	/** The viewport toolbar */
+	TSharedPtr<SContextualAnimViewportToolBar> ViewportToolbar;
+
 	/** Viewport client */
 	TSharedPtr<FContextualAnimViewportClient> ViewportClient;
 
diff --git a/Engine/Plugins/Experimental/Animation/ContextualAnimation/Source/ContextualAnimationEditor/Public/SContextualAnimViewportToolbar.h b/Engine/Plugins/Experimental/Animation/ContextualAnimation/Source/ContextualAnimationEditor/Public/SContextualAnimViewportToolbar.h
new file mode 100644
--- /dev/null
+++ b/Engine/Plugins/Experimental/Animation/ContextualAnimation/Source/ContextualAnimationEditor/Public/SContextualAnimViewportToolbar.h
@@ -0,0 +1,23 @@
+// Copyright Epic Games, Inc. All Rights Reserved.
+
+#pragma once
+
+#include "CoreMinimal.h"
+#include "SCommonEditorViewportToolbarBase.h"
+
+class SContextualAnimViewport;
+
+class SContextualAnimViewportToolBar : public SCommonEditorViewportToolbarBase
+{
+public:
+	SLATE_BEGIN_ARGS(SContextualAnimViewportToolBar)
+	{}
+
+	SLATE_END_ARGS()
+
+	void Construct(const FArguments& InArgs, TSharedPtr<SContextualAnimViewport> InViewport);
+
+	// ~SCommonEditorViewportToolbarBase interface
+	virtual TSharedRef<SWidget> GenerateShowMenu() const override;
+	// ~End of SCommonEditorViewportToolbarBase interface
+};
