Contextual Anim Editor (WIP): Fixed crash when removing a selection criterion while hitproxy is selected and removed inactive asset browser tab

diff --git a/Engine/Plugins/Experimental/Animation/ContextualAnimation/Source/ContextualAnimationEditor/Private/ContextualAnimAssetEditorToolkit.cpp b/Engine/Plugins/Experimental/Animation/ContextualAnimation/Source/ContextualAnimationEditor/Private/ContextualAnimAssetEditorToolkit.cpp
--- a/Engine/Plugins/Experimental/Animation/ContextualAnimation/Source/ContextualAnimationEditor/Private/ContextualAnimAssetEditorToolkit.cpp
+++ b/Engine/Plugins/Experimental/Animation/ContextualAnimation/Source/ContextualAnimationEditor/Private/ContextualAnimAssetEditorToolkit.cpp
@@ -96,7 +96,7 @@ void FContextualAnimAssetEditorToolkit::InitAssetEditor(const EToolkitMode::Type
 	EditingAssetWidget->OnFinishedChangingProperties().AddSP(this, &FContextualAnimAssetEditorToolkit::OnFinishedChangingProperties);
 
 	// Define Editor Layout
-	const TSharedRef<FTabManager::FLayout> StandaloneDefaultLayout = FTabManager::NewLayout("Standalone_ContextualAnimAnimEditor_Layout_v0.09")
+	const TSharedRef<FTabManager::FLayout> StandaloneDefaultLayout = FTabManager::NewLayout("Standalone_ContextualAnimAnimEditor_Layout_v0.10")
 		->AddArea
 		(
 			FTabManager::NewPrimaryArea()->SetOrientation(Orient_Vertical)
@@ -132,12 +132,6 @@ void FContextualAnimAssetEditorToolkit::InitAssetEditor(const EToolkitMode::Type
 							->AddTab(FContextualAnimEditorTabs::AssetDetailsID, ETabState::OpenedTab)
 							->AddTab(FContextualAnimEditorTabs::PreviewSettingsID, ETabState::OpenedTab)
 						)
-						->Split
-						(
-							FTabManager::NewStack()
-							->SetSizeCoefficient(0.3f)
-							->AddTab(FContextualAnimEditorTabs::AssetBrowserID, ETabState::OpenedTab)
-						)
 					)
 				)
 			)
diff --git a/Engine/Plugins/Experimental/Animation/ContextualAnimation/Source/ContextualAnimationEditor/Private/ContextualAnimEdMode.cpp b/Engine/Plugins/Experimental/Animation/ContextualAnimation/Source/ContextualAnimationEditor/Private/ContextualAnimEdMode.cpp
--- a/Engine/Plugins/Experimental/Animation/ContextualAnimation/Source/ContextualAnimationEditor/Private/ContextualAnimEdMode.cpp
+++ b/Engine/Plugins/Experimental/Animation/ContextualAnimation/Source/ContextualAnimationEditor/Private/ContextualAnimEdMode.cpp
@@ -268,34 +268,38 @@ bool FContextualAnimEdMode::InputDelta(FEditorViewportClient* InViewportClient,
 				const UContextualAnimSceneAsset* SceneAsset = ViewModel->GetSceneAsset();
 				if (const FContextualAnimTrack* AnimTrack = SceneAsset->GetAnimTrack(SelectedSelectionCriterionData.RoleName, SelectedSelectionCriterionData.VariantIdx))
 				{
-					if (UContextualAnimSelectionCriterion_TriggerArea* Spatial = Cast<UContextualAnimSelectionCriterion_TriggerArea>(AnimTrack->SelectionCriteria[SelectedSelectionCriterionData.CriterionIdx]))
+					if(AnimTrack->SelectionCriteria.IsValidIndex(SelectedSelectionCriterionData.CriterionIdx))
 					{
-						FMatrix WidgetCoordSystem;
-						GetCustomDrawingCoordinateSystem(WidgetCoordSystem, nullptr);
-						InDrag = WidgetCoordSystem.InverseTransformVector(InDrag);
+						if (UContextualAnimSelectionCriterion_TriggerArea* Spatial = Cast<UContextualAnimSelectionCriterion_TriggerArea>(AnimTrack->SelectionCriteria[SelectedSelectionCriterionData.CriterionIdx]))
+						{
+							FMatrix WidgetCoordSystem = FMatrix::Identity;
+							GetCustomDrawingCoordinateSystem(WidgetCoordSystem, nullptr);
 
-						FVector& Point = Spatial->PolygonPoints[SelectedSelectionCriterionData.DataIdx >= 4 ? SelectedSelectionCriterionData.DataIdx - 4 : SelectedSelectionCriterionData.DataIdx];
-						Point.X += InDrag.X;
-						Point.Y += InDrag.Y;
+							InDrag = WidgetCoordSystem.InverseTransformVector(InDrag);
 
-						if (InDrag.Z != 0.f)
-						{
-							if (SelectedSelectionCriterionData.DataIdx < 4)
+							FVector& Point = Spatial->PolygonPoints[SelectedSelectionCriterionData.DataIdx >= 4 ? SelectedSelectionCriterionData.DataIdx - 4 : SelectedSelectionCriterionData.DataIdx];
+							Point.X += InDrag.X;
+							Point.Y += InDrag.Y;
+
+							if (InDrag.Z != 0.f)
 							{
-								for (int32 Idx = 0; Idx < Spatial->PolygonPoints.Num(); Idx++)
+								if (SelectedSelectionCriterionData.DataIdx < 4)
 								{
-									Spatial->PolygonPoints[Idx].Z += InDrag.Z;
-								}
+									for (int32 Idx = 0; Idx < Spatial->PolygonPoints.Num(); Idx++)
+									{
+										Spatial->PolygonPoints[Idx].Z += InDrag.Z;
+									}
 
-								Spatial->Height = FMath::Max(Spatial->Height - InDrag.Z, 0.f);
-							}
-							else
-							{
-								Spatial->Height = FMath::Max(Spatial->Height + InDrag.Z, 0.f);
+									Spatial->Height = FMath::Max(Spatial->Height - InDrag.Z, 0.f);
+								}
+								else
+								{
+									Spatial->Height = FMath::Max(Spatial->Height + InDrag.Z, 0.f);
+								}
 							}
-						}
 
-						return true;
+							return true;
+						}
 					}
 				}
 			}
@@ -326,7 +330,20 @@ bool FContextualAnimEdMode::AllowWidgetMove()
 
 bool FContextualAnimEdMode::ShouldDrawWidget() const
 {
-	return SelectedSelectionCriterionData.IsValid();
+	if (ViewModel && SelectedSelectionCriterionData.IsValid())
+	{
+		const UContextualAnimSceneAsset* SceneAsset = ViewModel->GetSceneAsset();
+		if (const FContextualAnimTrack* AnimTrack = SceneAsset->GetAnimTrack(SelectedSelectionCriterionData.RoleName, SelectedSelectionCriterionData.VariantIdx))
+		{
+			const int32 Idx = SelectedSelectionCriterionData.CriterionIdx;
+			if (AnimTrack->SelectionCriteria.IsValidIndex(Idx) && AnimTrack->SelectionCriteria[Idx]->GetClass()->IsChildOf<UContextualAnimSelectionCriterion_TriggerArea>())
+			{
+				return true;
+			}
+		}
+	}
+
+	return false;
 }
 
 bool FContextualAnimEdMode::GetCustomDrawingCoordinateSystem(FMatrix& InMatrix, void* InData)
@@ -367,25 +384,28 @@ FVector FContextualAnimEdMode::GetWidgetLocation() const
 				const UContextualAnimSceneAsset* SceneAsset = ViewModel->GetSceneAsset();
 				if (const FContextualAnimTrack* AnimTrack = SceneAsset->GetAnimTrack(SelectedSelectionCriterionData.RoleName, SelectedSelectionCriterionData.VariantIdx))
 				{
-					if (const UContextualAnimSelectionCriterion_TriggerArea* Spatial = Cast<UContextualAnimSelectionCriterion_TriggerArea>(AnimTrack->SelectionCriteria[SelectedSelectionCriterionData.CriterionIdx]))
+					if (AnimTrack->SelectionCriteria.IsValidIndex(SelectedSelectionCriterionData.CriterionIdx))
 					{
-						FVector Location = FVector::ZeroVector;
-						if (SelectedSelectionCriterionData.DataIdx < 4)
-						{
-							Location = Spatial->PolygonPoints[SelectedSelectionCriterionData.DataIdx];
-						}
-						else
+						if (const UContextualAnimSelectionCriterion_TriggerArea* Spatial = Cast<UContextualAnimSelectionCriterion_TriggerArea>(AnimTrack->SelectionCriteria[SelectedSelectionCriterionData.CriterionIdx]))
 						{
-							Location = Spatial->PolygonPoints[SelectedSelectionCriterionData.DataIdx - 4] + FVector::UpVector * Spatial->Height;
-						}
+							FVector Location = FVector::ZeroVector;
+							if (SelectedSelectionCriterionData.DataIdx < 4)
+							{
+								Location = Spatial->PolygonPoints[SelectedSelectionCriterionData.DataIdx];
+							}
+							else
+							{
+								Location = Spatial->PolygonPoints[SelectedSelectionCriterionData.DataIdx - 4] + FVector::UpVector * Spatial->Height;
+							}
 
-						FTransform PrimaryActorTransform = FTransform::Identity;
-						if (const FContextualAnimSceneActorData* PreviewData = SceneInstance->FindSceneActorDataByRole(SceneAsset->GetPrimaryRole()))
-						{
-							PrimaryActorTransform = PreviewData->GetActor()->GetTransform();
-						}
+							FTransform PrimaryActorTransform = FTransform::Identity;
+							if (const FContextualAnimSceneActorData* PreviewData = SceneInstance->FindSceneActorDataByRole(SceneAsset->GetPrimaryRole()))
+							{
+								PrimaryActorTransform = PreviewData->GetActor()->GetTransform();
+							}
 
-						return PrimaryActorTransform.TransformPositionNoScale(Location);
+							return PrimaryActorTransform.TransformPositionNoScale(Location);
+						}
 					}
 				}
 			}
