Contextual Anim: Re-added old query function (for backward compatibility) and removed check when populating preview scene if preview actor class is not valid.

diff --git a/Engine/Plugins/Experimental/Animation/ContextualAnimation/Source/ContextualAnimation/Private/ContextualAnimSceneAsset.cpp b/Engine/Plugins/Experimental/Animation/ContextualAnimation/Source/ContextualAnimation/Private/ContextualAnimSceneAsset.cpp
--- a/Engine/Plugins/Experimental/Animation/ContextualAnimation/Source/ContextualAnimation/Private/ContextualAnimSceneAsset.cpp
+++ b/Engine/Plugins/Experimental/Animation/ContextualAnimation/Source/ContextualAnimation/Private/ContextualAnimSceneAsset.cpp
@@ -522,6 +522,33 @@ const FContextualAnimTrack* UContextualAnimSceneAsset::FindAnimTrackForRoleWithC
 	return Result;
 }
 
+bool UContextualAnimSceneAsset::Query(const FName& Role, FContextualAnimQueryResult& OutResult, const FContextualAnimQueryParams& QueryParams, const FTransform& ToWorldTransform) const
+{
+	FContextualAnimPrimaryActorData PrimaryActorData;
+	PrimaryActorData.Transform = ToWorldTransform;
+
+	FContextualAnimQuerierData QuerierData;
+	QuerierData.Transform = QueryParams.Querier.IsValid() ? QueryParams.Querier->GetActorTransform() : QueryParams.QueryTransform;
+
+	if(const FContextualAnimTrack* AnimTrack = FindFirstAnimTrackForRoleThatPassesSelectionCriteria(Role, PrimaryActorData, QuerierData))
+	{
+		OutResult.VariantIdx = AnimTrack->VariantIdx;
+		OutResult.Animation = AnimTrack->Animation;
+		OutResult.EntryTransform = AnimTrack->GetAlignmentTransformAtEntryTime() * ToWorldTransform;
+		OutResult.SyncTransform = AnimTrack->GetAlignmentTransformAtSyncTime() * ToWorldTransform;
+
+		if (QueryParams.bFindAnimStartTime)
+		{
+			const FVector LocalLocation = (QuerierData.Transform.GetRelativeTransform(ToWorldTransform)).GetLocation();
+			OutResult.AnimStartTime = AnimTrack->FindBestAnimStartTime(LocalLocation);
+		}
+
+		return true;
+	}
+	
+	return false;
+}
+
 FTransform UContextualAnimSceneAsset::GetAlignmentTransformForRoleRelativeToScenePivot(FName Role, int32 VariantIdx, float Time) const
 {
 	if (const FContextualAnimTrack* AnimTrack = GetAnimTrack(Role, VariantIdx))
diff --git a/Engine/Plugins/Experimental/Animation/ContextualAnimation/Source/ContextualAnimation/Public/ContextualAnimSceneAsset.h b/Engine/Plugins/Experimental/Animation/ContextualAnimation/Source/ContextualAnimation/Public/ContextualAnimSceneAsset.h
--- a/Engine/Plugins/Experimental/Animation/ContextualAnimation/Source/ContextualAnimation/Public/ContextualAnimSceneAsset.h
+++ b/Engine/Plugins/Experimental/Animation/ContextualAnimation/Source/ContextualAnimation/Public/ContextualAnimSceneAsset.h
@@ -135,12 +135,15 @@ public:
 	int32 GetTotalVariants() const;
 
 	UFUNCTION(BlueprintCallable, Category = "Contextual Anim|Scene Asset", meta = (DisplayName = "Get Anim Track"))
-	const FContextualAnimTrack& BP_GetAnimTrack(const FName& Role, int32 VariantIdx) const
+	const FContextualAnimTrack& BP_GetAnimTrack(FName Role, int32 VariantIdx) const
 	{
 		const FContextualAnimTrack* AnimTrack = GetAnimTrack(Role, VariantIdx);
 		return AnimTrack ? *AnimTrack : FContextualAnimTrack::EmptyTrack;
 	}
 
+	UFUNCTION(BlueprintCallable, Category = "Contextual Anim|Scene Asset")
+	bool Query(const FName& Role, FContextualAnimQueryResult& OutResult, const FContextualAnimQueryParams& QueryParams, const FTransform& ToWorldTransform) const;
+
 protected:
 
 	UPROPERTY(EditAnywhere, Category = "Defaults")
diff --git a/Engine/Plugins/Experimental/Animation/ContextualAnimation/Source/ContextualAnimation/Public/ContextualAnimTypes.h b/Engine/Plugins/Experimental/Animation/ContextualAnimation/Source/ContextualAnimation/Public/ContextualAnimTypes.h
--- a/Engine/Plugins/Experimental/Animation/ContextualAnimation/Source/ContextualAnimation/Public/ContextualAnimTypes.h
+++ b/Engine/Plugins/Experimental/Animation/ContextualAnimation/Source/ContextualAnimation/Public/ContextualAnimTypes.h
@@ -247,6 +247,33 @@ struct FContextualAnimQueryResult
 	FORCEINLINE bool IsValid() const { return VariantIdx != INDEX_NONE; }
 };
 
+/** Stores the parameters passed into query function */
+USTRUCT(BlueprintType)
+struct FContextualAnimQueryParams
+{
+	GENERATED_BODY()
+
+	UPROPERTY(EditAnywhere, BlueprintReadWrite, Category = "Defaults")
+	TWeakObjectPtr<const AActor> Querier;
+
+	UPROPERTY(EditAnywhere, BlueprintReadWrite, Category = "Defaults")
+	FTransform QueryTransform;
+
+	UPROPERTY(EditAnywhere, BlueprintReadWrite, Category = "Defaults")
+	bool bComplexQuery = false;
+
+	UPROPERTY(EditAnywhere, BlueprintReadWrite, Category = "Defaults")
+	bool bFindAnimStartTime = false;
+
+	FContextualAnimQueryParams() {}
+
+	FContextualAnimQueryParams(const AActor* InQuerier, bool bInComplexQuery, bool bInFindAnimStartTime)
+		: Querier(InQuerier), bComplexQuery(bInComplexQuery), bFindAnimStartTime(bInFindAnimStartTime) {}
+
+	FContextualAnimQueryParams(const FTransform& InQueryTransform, bool bInComplexQuery, bool bInFindAnimStartTime)
+		: QueryTransform(InQueryTransform), bComplexQuery(bInComplexQuery), bFindAnimStartTime(bInFindAnimStartTime) {}
+};
+
 USTRUCT(BlueprintType)
 struct FContextualAnimQuerierData
 {
diff --git a/Engine/Plugins/Experimental/Animation/ContextualAnimation/Source/ContextualAnimationEditor/Private/ContextualAnimViewModel.cpp b/Engine/Plugins/Experimental/Animation/ContextualAnimation/Source/ContextualAnimationEditor/Private/ContextualAnimViewModel.cpp
--- a/Engine/Plugins/Experimental/Animation/ContextualAnimation/Source/ContextualAnimationEditor/Private/ContextualAnimViewModel.cpp
+++ b/Engine/Plugins/Experimental/Animation/ContextualAnimation/Source/ContextualAnimationEditor/Private/ContextualAnimViewModel.cpp
@@ -202,7 +202,11 @@ void FContextualAnimViewModel::RefreshSequencerTracks()
 
 			// Spawn preview actor
 			AActor* PreviewActor = SpawnPreviewActor(AnimTrack);
-			check(PreviewActor);
+			
+			if(PreviewActor == nullptr)
+			{
+				return UE::ContextualAnim::EForEachResult::Continue;
+			}
 
 			// Set actor label so the track shows the name of the role it represents
 			//PreviewActor->SetActorLabel(Role.ToString(), false);
