ContextualAnim bug fixes (regressions): Fixed bug causing missing ik targets. Fixed bug causing crash when generating alignment tracks. Minor change: Removed const param from query function.

diff --git a/Engine/Plugins/Experimental/Animation/ContextualAnimation/Source/ContextualAnimation/Private/ContextualAnimSceneAsset.cpp b/Engine/Plugins/Experimental/Animation/ContextualAnimation/Source/ContextualAnimation/Private/ContextualAnimSceneAsset.cpp
--- a/Engine/Plugins/Experimental/Animation/ContextualAnimation/Source/ContextualAnimation/Private/ContextualAnimSceneAsset.cpp
+++ b/Engine/Plugins/Experimental/Animation/ContextualAnimation/Source/ContextualAnimation/Private/ContextualAnimSceneAsset.cpp
@@ -89,11 +89,11 @@ void UContextualAnimSceneAsset::GenerateAlignmentTracks()
 
 	for (FContextualAnimTracksContainer& Variant : Variants)
 	{
+		Variant.ScenePivots.Reset();
+
 		// Generate scene pivot for each alignment section
 		for (const FContextualAnimAlignmentSectionData& AligmentSection : AlignmentSections)
 		{
-			Variant.ScenePivots.Reset();
-
 			FTransform ScenePivot = FTransform::Identity;
 			FContextualAnimTrack* AnimTrack = Variant.Tracks.FindByPredicate([&AligmentSection](const FContextualAnimTrack& AnimTrack) { return AnimTrack.Role == AligmentSection.Origin; });
 			if (AnimTrack)
@@ -196,7 +196,7 @@ void UContextualAnimSceneAsset::GenerateIKTargetTracks()
 			const FContextualAnimIKTargetDefContainer* IKTargetDefContainerPtr = RoleToIKTargetDefsMap.Find(AnimTrack.Role);
 			if (IKTargetDefContainerPtr == nullptr || IKTargetDefContainerPtr->IKTargetDefs.Num() == 0)
 			{
-				return;
+				continue;
 			}
 
 			if (const UAnimMontage* Animation = AnimTrack.Animation)
@@ -522,7 +522,7 @@ const FContextualAnimTrack* UContextualAnimSceneAsset::FindAnimTrackForRoleWithC
 	return Result;
 }
 
-bool UContextualAnimSceneAsset::Query(const FName& Role, FContextualAnimQueryResult& OutResult, const FContextualAnimQueryParams& QueryParams, const FTransform& ToWorldTransform) const
+bool UContextualAnimSceneAsset::Query(FName Role, FContextualAnimQueryResult& OutResult, const FContextualAnimQueryParams& QueryParams, const FTransform& ToWorldTransform) const
 {
 	FContextualAnimPrimaryActorData PrimaryActorData;
 	PrimaryActorData.Transform = ToWorldTransform;
diff --git a/Engine/Plugins/Experimental/Animation/ContextualAnimation/Source/ContextualAnimation/Public/ContextualAnimSceneAsset.h b/Engine/Plugins/Experimental/Animation/ContextualAnimation/Source/ContextualAnimation/Public/ContextualAnimSceneAsset.h
--- a/Engine/Plugins/Experimental/Animation/ContextualAnimation/Source/ContextualAnimation/Public/ContextualAnimSceneAsset.h
+++ b/Engine/Plugins/Experimental/Animation/ContextualAnimation/Source/ContextualAnimation/Public/ContextualAnimSceneAsset.h
@@ -142,7 +142,7 @@ public:
 	}
 
 	UFUNCTION(BlueprintCallable, Category = "Contextual Anim|Scene Asset")
-	bool Query(const FName& Role, FContextualAnimQueryResult& OutResult, const FContextualAnimQueryParams& QueryParams, const FTransform& ToWorldTransform) const;
+	bool Query(FName Role, FContextualAnimQueryResult& OutResult, const FContextualAnimQueryParams& QueryParams, const FTransform& ToWorldTransform) const;
 
 protected:
 
