Contextual Anim (WIP): Refactored SceneAsset to work with AnimSequenceBase instead of only Montages (Montage still supported)

diff --git a/Engine/Plugins/Experimental/Animation/ContextualAnimation/Source/ContextualAnimation/Private/ContextualAnimSceneAsset.cpp b/Engine/Plugins/Experimental/Animation/ContextualAnimation/Source/ContextualAnimation/Private/ContextualAnimSceneAsset.cpp
--- a/Engine/Plugins/Experimental/Animation/ContextualAnimation/Source/ContextualAnimation/Private/ContextualAnimSceneAsset.cpp
+++ b/Engine/Plugins/Experimental/Animation/ContextualAnimation/Source/ContextualAnimation/Private/ContextualAnimSceneAsset.cpp
@@ -140,7 +140,7 @@ void UContextualAnimSceneAsset::GenerateAlignmentTracks()
 				AnimTrack.AlignmentData.Tracks.AnimationTracks.AddZeroed();
 			}
 
-			if (const UAnimMontage* Animation = AnimTrack.Animation)
+			if (const UAnimSequenceBase* Animation = AnimTrack.Animation)
 			{
 				float Time = 0.f;
 				float EndTime = Animation->GetPlayLength();
@@ -199,7 +199,7 @@ void UContextualAnimSceneAsset::GenerateIKTargetTracks()
 				continue;
 			}
 
-			if (const UAnimMontage* Animation = AnimTrack.Animation)
+			if (const UAnimSequenceBase* Animation = AnimTrack.Animation)
 			{
 				UE_LOG(LogContextualAnim, Log, TEXT("%s Generating IK Target Tracks. Animation: %s"), *GetNameSafe(this), *GetNameSafe(Animation));
 
@@ -250,7 +250,7 @@ void UContextualAnimSceneAsset::GenerateIKTargetTracks()
 						// Find TargetBoneIndex. Note that we add TargetBoneIndex even if it is INDEX_NONE. In this case, my bone will be relative to the origin of the target actor. 
 						// This is to support cases where the target actor doesn't have animation or TargetBoneName is None
 						FName TargetBoneName = IKTargetDef.TargetBoneName;
-						const UAnimMontage* TargetAnimation = DataPtr->TargetAnimTrackPtr->Animation;
+						const UAnimSequenceBase* TargetAnimation = DataPtr->TargetAnimTrackPtr->Animation;
 						const int32 TargetBoneIndex = TargetAnimation ? TargetAnimation->GetSkeleton()->GetReferenceSkeleton().FindBoneIndex(TargetBoneName) : INDEX_NONE;
 						if (TargetBoneIndex == INDEX_NONE)
 						{
@@ -298,7 +298,7 @@ void UContextualAnimSceneAsset::GenerateIKTargetTracks()
 							FCSPose<FCompactPose> OtherComponentSpacePose;
 							TArray<FBoneIndexType> OtherRequiredBoneIndexArray;
 							FBoneContainer OtherBoneContainer;
-							const UAnimMontage* OtherAnimation = Data.Value.TargetAnimTrackPtr->Animation;
+							const UAnimSequenceBase* OtherAnimation = Data.Value.TargetAnimTrackPtr->Animation;
 							if (OtherAnimation)
 							{
 								// Prepare array with the indices of the bones to extract from target animation
@@ -373,7 +373,7 @@ void UContextualAnimSceneAsset::UpdateRadius()
 		});
 }
 
-FName UContextualAnimSceneAsset::FindRoleByAnimation(const UAnimMontage* Animation) const
+FName UContextualAnimSceneAsset::FindRoleByAnimation(const UAnimSequenceBase* Animation) const
 {
 	FName Result = NAME_None;
 
@@ -524,6 +524,8 @@ const FContextualAnimTrack* UContextualAnimSceneAsset::FindAnimTrackForRoleWithC
 
 bool UContextualAnimSceneAsset::Query(FName Role, FContextualAnimQueryResult& OutResult, const FContextualAnimQueryParams& QueryParams, const FTransform& ToWorldTransform) const
 {
+	//@TODO: Kept around only to do not break existing content. It will go away in the future.
+
 	FContextualAnimPrimaryActorData PrimaryActorData;
 	PrimaryActorData.Transform = ToWorldTransform;
 
@@ -533,7 +535,7 @@ bool UContextualAnimSceneAsset::Query(FName Role, FContextualAnimQueryResult& Ou
 	if(const FContextualAnimTrack* AnimTrack = FindFirstAnimTrackForRoleThatPassesSelectionCriteria(Role, PrimaryActorData, QuerierData))
 	{
 		OutResult.VariantIdx = AnimTrack->VariantIdx;
-		OutResult.Animation = AnimTrack->Animation;
+		OutResult.Animation = Cast<UAnimMontage>(AnimTrack->Animation); // Its ok to cast here to montage. Content using this function uses montage
 		OutResult.EntryTransform = AnimTrack->GetAlignmentTransformAtEntryTime() * ToWorldTransform;
 		OutResult.SyncTransform = AnimTrack->GetAlignmentTransformAtSyncTime() * ToWorldTransform;
 
@@ -576,7 +578,7 @@ FTransform UContextualAnimSceneAsset::GetIKTargetTransformForRoleAtTime(FName Ro
 	return FTransform::Identity;
 }
 
-int32 UContextualAnimSceneAsset::FindVariantIdx(FName Role, UAnimMontage* Animation) const
+int32 UContextualAnimSceneAsset::FindVariantIdx(FName Role, UAnimSequenceBase* Animation) const
 {
 	int32 Result = INDEX_NONE;
 
diff --git a/Engine/Plugins/Experimental/Animation/ContextualAnimation/Source/ContextualAnimation/Private/ContextualAnimSceneInstance.cpp b/Engine/Plugins/Experimental/Animation/ContextualAnimation/Source/ContextualAnimation/Private/ContextualAnimSceneInstance.cpp
--- a/Engine/Plugins/Experimental/Animation/ContextualAnimation/Source/ContextualAnimation/Private/ContextualAnimSceneInstance.cpp
+++ b/Engine/Plugins/Experimental/Animation/ContextualAnimation/Source/ContextualAnimation/Private/ContextualAnimSceneInstance.cpp
@@ -38,23 +38,14 @@ USkeletalMeshComponent* FContextualAnimSceneActorData::GetSkeletalMeshComponent(
 
 FAnimMontageInstance* FContextualAnimSceneActorData::GetAnimMontageInstance() const
 {
-	if (const UAnimMontage* Animation = AnimTrackPtr->Animation)
+	if (UAnimInstance* AnimInstance = GetAnimInstance())
 	{
-		if (UAnimInstance* AnimInstance = GetAnimInstance())
-		{
-			return AnimInstance->GetActiveInstanceForMontage(Animation);
-		}
+		return AnimInstance->GetActiveMontageInstance();
 	}
 
 	return nullptr;
 }
 
-UAnimMontage* FContextualAnimSceneActorData::GetAnimMontage() const
-{
-	const FAnimMontageInstance* MontageInstance = GetAnimMontageInstance();
-	return MontageInstance ? MontageInstance->Montage : nullptr;
-}
-
 float FContextualAnimSceneActorData::GetAnimTime() const
 {
 	const FAnimMontageInstance* MontageInstance = GetAnimMontageInstance();
@@ -88,7 +79,7 @@ const FContextualAnimIKTargetDefContainer& FContextualAnimSceneActorData::GetIKT
 void UContextualAnimSceneInstance::BreakContextualAnimSceneActorData(const FContextualAnimSceneActorData& SceneActorData, AActor*& Actor, UAnimMontage*& Montage, float& AnimTime, int32& CurrentSectionIndex, FName& CurrentSectionName)
 {
 	Actor = SceneActorData.GetActor();
-	Montage = const_cast<UAnimMontage*>(SceneActorData.GetAnimMontage());
+	Montage = Cast<UAnimMontage>(SceneActorData.GetAnimTrack().Animation);
 	AnimTime = SceneActorData.GetAnimTime();
 	CurrentSectionIndex = SceneActorData.GetCurrentSectionIndex();
 	CurrentSectionName = SceneActorData.GetCurrentSection();
@@ -125,16 +116,41 @@ AActor* UContextualAnimSceneInstance::GetActorByRole(FName Role) const
 void UContextualAnimSceneInstance::Join(FContextualAnimSceneActorData& SceneActorData)
 {
 	AActor* Actor = SceneActorData.GetActor();
-	if (UAnimInstance* AnimInstance = SceneActorData.GetAnimInstance())
+	if (Actor == nullptr)
 	{
-		UAnimMontage* Animation = SceneActorData.GetAnimTrack().Animation;
-		AnimInstance->Montage_Play(Animation, 1.f, EMontagePlayReturnType::MontageLength, SceneActorData.GetAnimStartTime());
+		return;
+	}
+
+	if (UAnimSequenceBase* Animation = SceneActorData.GetAnimTrack().Animation)
+	{
+		if (UAnimInstance* AnimInstance = SceneActorData.GetAnimInstance())
+		{
+			// Keep montage support for now but might go away soon
+			if (UAnimMontage* AnimMontage = Cast<UAnimMontage>(Animation))
+			{
+				AnimInstance->Montage_Play(AnimMontage, 1.f, EMontagePlayReturnType::MontageLength, SceneActorData.GetAnimStartTime());
+			}
+			else
+			{
+				//@TODO: Expose all these on the AnimTrack
+				const FName SlotName = FName(TEXT("DefaultSlot"));
+				const float BlendInTime = 0.25f;
+				const float BlendOutTime = 0.25f;
+				const float InPlayRate = 1.f;
+				const int32 LoopCount = 1;
+				const float BlendOutTriggerTime = -1.f;
+				const float InTimeToStartMontageAt = SceneActorData.GetAnimStartTime();
+				AnimInstance->PlaySlotAnimationAsDynamicMontage(SceneActorData.GetAnimTrack().Animation, SlotName, BlendInTime, BlendOutTime, InPlayRate, LoopCount, BlendOutTriggerTime, InTimeToStartMontageAt);
+			}
+
+			AnimInstance->OnPlayMontageNotifyBegin.AddUniqueDynamic(this, &UContextualAnimSceneInstance::OnNotifyBeginReceived);
+			AnimInstance->OnPlayMontageNotifyEnd.AddUniqueDynamic(this, &UContextualAnimSceneInstance::OnNotifyEndReceived);
+			AnimInstance->OnMontageBlendingOut.AddUniqueDynamic(this, &UContextualAnimSceneInstance::OnMontageBlendingOut);
+		}
 
-		AnimInstance->OnPlayMontageNotifyBegin.AddUniqueDynamic(this, &UContextualAnimSceneInstance::OnNotifyBeginReceived);
-		AnimInstance->OnPlayMontageNotifyEnd.AddUniqueDynamic(this, &UContextualAnimSceneInstance::OnNotifyEndReceived);
-		AnimInstance->OnMontageBlendingOut.AddUniqueDynamic(this, &UContextualAnimSceneInstance::OnMontageBlendingOut);
 
-		if(SceneActorData.GetAnimTrack().bRequireFlyingMode)
+		//@TODO: Temp, until we have a way to switch between movement mode using AnimNotifyState
+		if (SceneActorData.GetAnimTrack().bRequireFlyingMode)
 		{
 			if (UCharacterMovementComponent* CharacterMovementComp = Actor->FindComponentByClass<UCharacterMovementComponent>())
 			{
@@ -317,10 +333,32 @@ void UContextualAnimSceneInstance::OnMontageBlendingOut(UAnimMontage* Montage, b
 	{
 		if (UAnimInstance* AnimInstance = Binding.GetAnimInstance())
 		{
-			if(AnimInstance->Montage_IsPlaying(Binding.GetAnimTrack().Animation))
+			// Keep montage support for now but might go away soon
+			if (UAnimMontage* AnimMontage = Cast<UAnimMontage>(Binding.GetAnimTrack().Animation))
 			{
-				bShouldEnd = false;
-				break;
+				if (AnimInstance->Montage_IsPlaying(AnimMontage))
+				{
+					bShouldEnd = false;
+					break;
+				}
+			}
+			else
+			{
+				for (const FAnimMontageInstance* MontageInstance : AnimInstance->MontageInstances)
+				{
+					// When the animation is not a Montage, we still play it as a Montage. This dynamically created Montage has a single slot and single segment.
+					if (MontageInstance && MontageInstance->IsPlaying())
+					{
+						if(MontageInstance->Montage->SlotAnimTracks.Num() > 0 && MontageInstance->Montage->SlotAnimTracks[0].AnimTrack.AnimSegments.Num() > 0)
+						{
+							if(MontageInstance->Montage->SlotAnimTracks[0].AnimTrack.AnimSegments[0].GetAnimReference() == Binding.GetAnimTrack().Animation)
+							{
+								bShouldEnd = false;
+								break;
+							}
+						}
+					}
+				}
 			}
 		}
 	}
diff --git a/Engine/Plugins/Experimental/Animation/ContextualAnimation/Source/ContextualAnimation/Private/ContextualAnimUtilities.cpp b/Engine/Plugins/Experimental/Animation/ContextualAnimation/Source/ContextualAnimation/Private/ContextualAnimUtilities.cpp
--- a/Engine/Plugins/Experimental/Animation/ContextualAnimation/Source/ContextualAnimation/Private/ContextualAnimUtilities.cpp
+++ b/Engine/Plugins/Experimental/Animation/ContextualAnimation/Source/ContextualAnimation/Private/ContextualAnimUtilities.cpp
@@ -138,7 +138,7 @@ void UContextualAnimUtilities::DrawDebugScene(const UWorld* World, const UContex
 		{
 			const FTransform Transform = (SceneAsset->GetMeshToComponentForRole(AnimTrack.Role) * AnimTrack.GetAlignmentTransformAtTime(Time)) * ToWorldTransform;
 			
-			if (const UAnimMontage* Animation = AnimTrack.Animation)
+			if (const UAnimSequenceBase* Animation = AnimTrack.Animation)
 			{
 				DrawDebugPose(World, Animation, Time, Transform, Color, LifeTime, Thickness);
 			}
diff --git a/Engine/Plugins/Experimental/Animation/ContextualAnimation/Source/ContextualAnimation/Public/ContextualAnimSceneAsset.h b/Engine/Plugins/Experimental/Animation/ContextualAnimation/Source/ContextualAnimation/Public/ContextualAnimSceneAsset.h
--- a/Engine/Plugins/Experimental/Animation/ContextualAnimation/Source/ContextualAnimation/Public/ContextualAnimSceneAsset.h
+++ b/Engine/Plugins/Experimental/Animation/ContextualAnimation/Source/ContextualAnimation/Public/ContextualAnimSceneAsset.h
@@ -111,9 +111,9 @@ public:
 	const FContextualAnimTrack* GetAnimTrack(const FName& Role, int32 VariantIdx) const;
 
 	UFUNCTION(BlueprintCallable, Category = "Contextual Anim|Scene Asset")
-	int32 FindVariantIdx(FName Role, UAnimMontage* Animation) const;
+	int32 FindVariantIdx(FName Role, UAnimSequenceBase* Animation) const;
 
-	FName FindRoleByAnimation(const UAnimMontage* Animation) const;
+	FName FindRoleByAnimation(const UAnimSequenceBase* Animation) const;
 
 	const FContextualAnimTrack* FindFirstAnimTrackForRoleThatPassesSelectionCriteria(const FName& Role, const FContextualAnimPrimaryActorData& PrimaryActorData, const FContextualAnimQuerierData& QuerierData) const;
 	
@@ -141,6 +141,7 @@ public:
 		return AnimTrack ? *AnimTrack : FContextualAnimTrack::EmptyTrack;
 	}
 
+	//@TODO: Kept around only to do not break existing content. It will go away in the future.
 	UFUNCTION(BlueprintCallable, Category = "Contextual Anim|Scene Asset")
 	bool Query(FName Role, FContextualAnimQueryResult& OutResult, const FContextualAnimQueryParams& QueryParams, const FTransform& ToWorldTransform) const;
 
diff --git a/Engine/Plugins/Experimental/Animation/ContextualAnimation/Source/ContextualAnimation/Public/ContextualAnimSceneInstance.h b/Engine/Plugins/Experimental/Animation/ContextualAnimation/Source/ContextualAnimation/Public/ContextualAnimSceneInstance.h
--- a/Engine/Plugins/Experimental/Animation/ContextualAnimation/Source/ContextualAnimation/Public/ContextualAnimSceneInstance.h
+++ b/Engine/Plugins/Experimental/Animation/ContextualAnimation/Source/ContextualAnimation/Public/ContextualAnimSceneInstance.h
@@ -53,9 +53,6 @@ struct CONTEXTUALANIMATION_API FContextualAnimSceneActorData
 	/** Returns the ActiveMontageInstance or null in the case of static actors */
 	FAnimMontageInstance* GetAnimMontageInstance() const;
 
-	/** Returns the Animation this actor is playing or null in the case of static actors */
-	UAnimMontage* GetAnimMontage() const;
-
 	UAnimInstance* GetAnimInstance() const;
 
 	USkeletalMeshComponent* GetSkeletalMeshComponent() const;
diff --git a/Engine/Plugins/Experimental/Animation/ContextualAnimation/Source/ContextualAnimation/Public/ContextualAnimTypes.h b/Engine/Plugins/Experimental/Animation/ContextualAnimation/Source/ContextualAnimation/Public/ContextualAnimTypes.h
--- a/Engine/Plugins/Experimental/Animation/ContextualAnimation/Source/ContextualAnimation/Public/ContextualAnimTypes.h
+++ b/Engine/Plugins/Experimental/Animation/ContextualAnimation/Source/ContextualAnimation/Public/ContextualAnimTypes.h
@@ -10,7 +10,7 @@
 CONTEXTUALANIMATION_API DECLARE_LOG_CATEGORY_EXTERN(LogContextualAnim, Log, All);
 
 class AActor;
-class UAnimMontage;
+class UAnimSequenceBase;
 class UContextualAnimSelectionCriterion;
 
 namespace UE 
@@ -61,7 +61,7 @@ struct CONTEXTUALANIMATION_API FContextualAnimTrack
 	GENERATED_BODY()
 
 	UPROPERTY(EditAnywhere, BlueprintReadWrite, Category = "Defaults")
-	UAnimMontage* Animation = nullptr;
+	UAnimSequenceBase* Animation = nullptr;
 
 	UPROPERTY(EditAnywhere, BlueprintReadWrite, Category = "Defaults")
 	float AnimMaxStartTime = 0.f;
diff --git a/Engine/Plugins/Experimental/Animation/ContextualAnimation/Source/ContextualAnimationEditor/Private/ContextualAnimMovieSceneNotifySection.cpp b/Engine/Plugins/Experimental/Animation/ContextualAnimation/Source/ContextualAnimationEditor/Private/ContextualAnimMovieSceneNotifySection.cpp
--- a/Engine/Plugins/Experimental/Animation/ContextualAnimation/Source/ContextualAnimationEditor/Private/ContextualAnimMovieSceneNotifySection.cpp
+++ b/Engine/Plugins/Experimental/Animation/ContextualAnimation/Source/ContextualAnimationEditor/Private/ContextualAnimMovieSceneNotifySection.cpp
@@ -2,7 +2,7 @@
 
 #include "ContextualAnimMovieSceneNotifySection.h"
 #include "ContextualAnimMovieSceneNotifyTrack.h"
-#include "Animation/AnimMontage.h"
+#include "Animation/AnimSequenceBase.h"
 #include "MovieScene.h"
 
 UContextualAnimMovieSceneNotifyTrack* UContextualAnimMovieSceneNotifySection::GetOwnerTrack() const
@@ -22,7 +22,7 @@ void UContextualAnimMovieSceneNotifySection::Initialize(const FAnimNotifyEvent&
 
 FAnimNotifyEvent* UContextualAnimMovieSceneNotifySection::GetAnimNotifyEvent() const
 {
-	UAnimMontage& Animation = GetOwnerTrack()->GetAnimation();
+	UAnimSequenceBase& Animation = GetOwnerTrack()->GetAnimation();
 	for (FAnimNotifyEvent& NotifyEvent : Animation.Notifies)
 	{
 		if (NotifyEvent.Guid == AnimNotifyEventGuid)
@@ -36,7 +36,7 @@ FAnimNotifyEvent* UContextualAnimMovieSceneNotifySection::GetAnimNotifyEvent() c
 
 UAnimNotifyState* UContextualAnimMovieSceneNotifySection::GetAnimNotifyState() const
 {
-	UAnimMontage& Animation = GetOwnerTrack()->GetAnimation();
+	UAnimSequenceBase& Animation = GetOwnerTrack()->GetAnimation();
 	if(const FAnimNotifyEvent* NotifyEventPtr = GetAnimNotifyEvent())
 	{
 		return NotifyEventPtr->NotifyStateClass;
diff --git a/Engine/Plugins/Experimental/Animation/ContextualAnimation/Source/ContextualAnimationEditor/Private/ContextualAnimMovieSceneNotifyTrack.cpp b/Engine/Plugins/Experimental/Animation/ContextualAnimation/Source/ContextualAnimationEditor/Private/ContextualAnimMovieSceneNotifyTrack.cpp
--- a/Engine/Plugins/Experimental/Animation/ContextualAnimation/Source/ContextualAnimationEditor/Private/ContextualAnimMovieSceneNotifyTrack.cpp
+++ b/Engine/Plugins/Experimental/Animation/ContextualAnimation/Source/ContextualAnimationEditor/Private/ContextualAnimMovieSceneNotifyTrack.cpp
@@ -2,7 +2,7 @@
 
 #include "ContextualAnimMovieSceneNotifyTrack.h"
 #include "ContextualAnimMovieSceneNotifySection.h"
-#include "Animation/AnimMontage.h"
+#include "Animation/AnimSequenceBase.h"
 #include "Animation/AnimNotifies/AnimNotify.h"
 #include "Animation/AnimNotifies/AnimNotifyState.h"
 #include "ContextualAnimViewModel.h"
@@ -15,7 +15,7 @@ FContextualAnimViewModel& UContextualAnimMovieSceneNotifyTrack::GetViewModel() c
 	return GetTypedOuter<UMovieScene>()->GetTypedOuter<UContextualAnimMovieSceneSequence>()->GetViewModel();
 }
 
-void UContextualAnimMovieSceneNotifyTrack::Initialize(UAnimMontage& AnimationRef, const FAnimNotifyTrack& NotifyTrack)
+void UContextualAnimMovieSceneNotifyTrack::Initialize(UAnimSequenceBase& AnimationRef, const FAnimNotifyTrack& NotifyTrack)
 {
 	// Cache animation this track was created from
 	Animation = &AnimationRef;
@@ -38,7 +38,7 @@ void UContextualAnimMovieSceneNotifyTrack::Initialize(UAnimMontage& AnimationRef
 	}
 }
 
-UAnimMontage& UContextualAnimMovieSceneNotifyTrack::GetAnimation() const 
+UAnimSequenceBase& UContextualAnimMovieSceneNotifyTrack::GetAnimation() const
 { 
 	check(Animation.IsValid());
 	return *Animation.Get();
@@ -84,7 +84,7 @@ void UContextualAnimMovieSceneNotifyTrack::RemoveSection(UMovieSceneSection& Sec
 	// Remove Notify from the animation
 	UContextualAnimMovieSceneNotifySection* NotifySection = CastChecked<UContextualAnimMovieSceneNotifySection>(&Section);
 
-	UAnimMontage& Anim = GetAnimation();
+	UAnimSequenceBase& Anim = GetAnimation();
 
 	FGuid Guid = NotifySection->GetAnimNotifyEventGuid();
 	const int32 TotalRemoved = Anim.Notifies.RemoveAll([&Guid](const FAnimNotifyEvent& Event) { return Event.Guid == Guid; });
@@ -114,14 +114,14 @@ EMovieSceneSectionMovedResult UContextualAnimMovieSceneNotifyTrack::OnSectionMov
 
 		if (FAnimNotifyEvent* NotifyEventPtr = NotifySection->GetAnimNotifyEvent())
 		{
-			UAnimMontage* Anim = nullptr;
+			UAnimSequenceBase* Anim = nullptr;
 			if (NotifyEventPtr->NotifyStateClass)
 			{
-				Anim = NotifyEventPtr->NotifyStateClass->GetTypedOuter<UAnimMontage>();
+				Anim = NotifyEventPtr->NotifyStateClass->GetTypedOuter<UAnimSequenceBase>();
 			}
 			else if (NotifyEventPtr->Notify)
 			{
-				Anim = NotifyEventPtr->Notify->GetTypedOuter<UAnimMontage>();
+				Anim = NotifyEventPtr->Notify->GetTypedOuter<UAnimSequenceBase>();
 			}
 
 			if(Anim)
diff --git a/Engine/Plugins/Experimental/Animation/ContextualAnimation/Source/ContextualAnimationEditor/Private/ContextualAnimMovieSceneNotifyTrackEditor.cpp b/Engine/Plugins/Experimental/Animation/ContextualAnimation/Source/ContextualAnimationEditor/Private/ContextualAnimMovieSceneNotifyTrackEditor.cpp
--- a/Engine/Plugins/Experimental/Animation/ContextualAnimation/Source/ContextualAnimationEditor/Private/ContextualAnimMovieSceneNotifyTrackEditor.cpp
+++ b/Engine/Plugins/Experimental/Animation/ContextualAnimation/Source/ContextualAnimationEditor/Private/ContextualAnimMovieSceneNotifyTrackEditor.cpp
@@ -329,7 +329,7 @@ void FContextualAnimMovieSceneNotifyTrackEditor::AddNewNotifyTrack(TArray<FGuid>
 
 	for (const FGuid& ObjectBinding : ObjectBindings)
 	{
-		UAnimMontage* Animation = GetMovieSceneSequence().GetViewModel().FindAnimationByGuid(ObjectBinding);
+		UAnimSequenceBase* Animation = GetMovieSceneSequence().GetViewModel().FindAnimationByGuid(ObjectBinding);
 		check(Animation);
 
 		// Copied from AnimTimelineTrack_Notifies.cpp/FAnimTimelineTrack_Notifies::AddTrack()
@@ -368,7 +368,7 @@ UContextualAnimMovieSceneNotifySection* FContextualAnimMovieSceneNotifyTrackEdit
 	TSharedPtr<ISequencer> SequencerPtr = GetSequencer();
 	if (SequencerPtr.IsValid())
 	{
-		UAnimMontage& Animation = Track->GetAnimation();
+		UAnimSequenceBase& Animation = Track->GetAnimation();
 
 		const FName TrackName = FName(Track->GetDisplayName().ToString());
 		const int32 TrackIndex = Animation.AnimNotifyTracks.IndexOfByPredicate([TrackName](const FAnimNotifyTrack& AnimNotifyTrack) { return AnimNotifyTrack.TrackName == TrackName; });
diff --git a/Engine/Plugins/Experimental/Animation/ContextualAnimation/Source/ContextualAnimationEditor/Private/ContextualAnimViewModel.cpp b/Engine/Plugins/Experimental/Animation/ContextualAnimation/Source/ContextualAnimationEditor/Private/ContextualAnimViewModel.cpp
--- a/Engine/Plugins/Experimental/Animation/ContextualAnimation/Source/ContextualAnimationEditor/Private/ContextualAnimViewModel.cpp
+++ b/Engine/Plugins/Experimental/Animation/ContextualAnimation/Source/ContextualAnimationEditor/Private/ContextualAnimViewModel.cpp
@@ -74,9 +74,9 @@ void FContextualAnimViewModel::Initialize(UContextualAnimSceneAsset* InSceneAsse
 	RefreshSequencerTracks();
 }
 
-UAnimMontage* FContextualAnimViewModel::FindAnimationByGuid(const FGuid& Guid) const
+UAnimSequenceBase* FContextualAnimViewModel::FindAnimationByGuid(const FGuid& Guid) const
 {
-	return SceneInstance.IsValid() ? const_cast<UAnimMontage*>(SceneInstance->FindSceneActorDataByGuid(Guid)->GetAnimMontage()) : nullptr;
+	return SceneInstance.IsValid() ? SceneInstance->FindSceneActorDataByGuid(Guid)->GetAnimTrack().Animation : nullptr;
 }
 
 void FContextualAnimViewModel::CreateSequencer()
@@ -159,7 +159,7 @@ void FContextualAnimViewModel::RefreshSequencerTracks()
 {
 	// Remove movie scene tracks and destroy existing actors (if any)
 
-	for(UAnimMontage* Anim : AnimsBeingEdited)
+	for(UAnimSequenceBase* Anim : AnimsBeingEdited)
 	{
 		Anim->UnregisterOnNotifyChanged(this);
 	}
@@ -219,8 +219,8 @@ void FContextualAnimViewModel::RefreshSequencerTracks()
 
 			const FGuid& Guid = Guids[0];
 
-			UAnimMontage* AnimMontage = AnimTrack.Animation;
-			if (AnimMontage)
+			UAnimSequenceBase* Animation = AnimTrack.Animation;
+			if (Animation)
 			{
 				// Add Animation Track
 				{
@@ -228,14 +228,14 @@ void FContextualAnimViewModel::RefreshSequencerTracks()
 					UMovieSceneEventTrack* MovieSceneAnimTrack = MovieSceneSequence->GetMovieScene()->AddTrack<UMovieSceneEventTrack>(Guid);
 					check(MovieSceneAnimTrack);
 
-					MovieSceneAnimTrack->SetDisplayName(FText::FromString(GetNameSafe(AnimMontage)));
+					MovieSceneAnimTrack->SetDisplayName(FText::FromString(GetNameSafe(Animation)));
 
 					UMovieSceneSection* NewSection = NewObject<UMovieSceneSection>(MovieSceneAnimTrack, UMovieSceneEventRepeaterSection::StaticClass(), NAME_None, RF_Transactional);
 					check(NewSection);
 
 					FFrameRate TickResolution = MovieSceneSequence->GetMovieScene()->GetTickResolution();
 					FFrameNumber StartFrame(0);
-					FFrameNumber EndFrame = (AnimMontage->GetPlayLength() * TickResolution).RoundToFrame();
+					FFrameNumber EndFrame = (Animation->GetPlayLength() * TickResolution).RoundToFrame();
 					NewSection->SetRange(TRange<FFrameNumber>::Exclusive(StartFrame, EndFrame));
 
 					MovieSceneAnimTrack->AddSection(*NewSection);
@@ -243,19 +243,19 @@ void FContextualAnimViewModel::RefreshSequencerTracks()
 
 				// Add Notify Tracks
 				{
-					for (const FAnimNotifyTrack& NotifyTrack : AnimMontage->AnimNotifyTracks)
+					for (const FAnimNotifyTrack& NotifyTrack : Animation->AnimNotifyTracks)
 					{
 						UContextualAnimMovieSceneNotifyTrack* Track = MovieSceneSequence->GetMovieScene()->AddTrack<UContextualAnimMovieSceneNotifyTrack>(Guid);
 						check(Track);
 
-						Track->Initialize(*AnimMontage, NotifyTrack);
+						Track->Initialize(*Animation, NotifyTrack);
 					}
 
 					// Listen for when the notifies in the animation changes, so we can refresh the notify sections here
-					AnimMontage->RegisterOnNotifyChanged(UAnimSequenceBase::FOnNotifyChanged::CreateSP(this, &FContextualAnimViewModel::OnAnimNotifyChanged, AnimMontage));
+					Animation->RegisterOnNotifyChanged(UAnimSequenceBase::FOnNotifyChanged::CreateSP(this, &FContextualAnimViewModel::OnAnimNotifyChanged, Animation));
 				}
 
-				AnimsBeingEdited.Add(AnimMontage);
+				AnimsBeingEdited.Add(Animation);
 			}
 
 			StartSceneParams.RoleToActorMap.Add(Role, PreviewActor);
@@ -366,7 +366,7 @@ void FContextualAnimViewModel::SequencerTimeChanged()
 	{
 		const USkeletalMeshComponent* SkelMeshComp = UContextualAnimUtilities::TryGetSkeletalMeshComponent(SceneActorData.GetActor());
 
-		const FTransform RootTransform = UContextualAnimUtilities::ExtractRootTransformFromAnimation(SceneActorData.GetAnimMontage(), Time);
+		const FTransform RootTransform = UContextualAnimUtilities::ExtractRootTransformFromAnimation(SceneActorData.GetAnimTrack().Animation, Time);
 		const FTransform StartTransform = SkelMeshComp->GetRelativeTransform().Inverse() * RootTransform;
 
 		SceneActorData.GetActor()->SetActorLocationAndRotation(StartTransform.GetLocation(), StartTransform.GetRotation());
@@ -462,7 +462,7 @@ void FContextualAnimViewModel::SequencerDataChanged(EMovieSceneDataChangeType Da
 	}
 }
 
-void FContextualAnimViewModel::OnAnimNotifyChanged(UAnimMontage* Animation)
+void FContextualAnimViewModel::OnAnimNotifyChanged(UAnimSequenceBase* Animation)
 {
 	// Do not refresh sequencer tracks if the change to the notifies came from us
 	if (bUpdatingAnimationFromSequencer)
@@ -475,7 +475,7 @@ void FContextualAnimViewModel::OnAnimNotifyChanged(UAnimMontage* Animation)
 	RefreshSequencerTracks();
 }
 
-void FContextualAnimViewModel::AnimationModified(UAnimMontage& Animation)
+void FContextualAnimViewModel::AnimationModified(UAnimSequenceBase& Animation)
 {
 	TGuardValue<bool> UpdateGuard(bUpdatingAnimationFromSequencer, true);
 
diff --git a/Engine/Plugins/Experimental/Animation/ContextualAnimation/Source/ContextualAnimationEditor/Public/ContextualAnimMovieSceneNotifyTrack.h b/Engine/Plugins/Experimental/Animation/ContextualAnimation/Source/ContextualAnimationEditor/Public/ContextualAnimMovieSceneNotifyTrack.h
--- a/Engine/Plugins/Experimental/Animation/ContextualAnimation/Source/ContextualAnimationEditor/Public/ContextualAnimMovieSceneNotifyTrack.h
+++ b/Engine/Plugins/Experimental/Animation/ContextualAnimation/Source/ContextualAnimationEditor/Public/ContextualAnimMovieSceneNotifyTrack.h
@@ -6,7 +6,7 @@
 #include "MovieSceneNameableTrack.h"
 #include "ContextualAnimMovieSceneNotifyTrack.generated.h"
 
-class UAnimMontage;
+class UAnimSequenceBase;
 struct FAnimNotifyTrack;
 class UContextualAnimMovieSceneNotifySection;
 class FContextualAnimViewModel;
@@ -25,7 +25,7 @@ public:
 	 * Initialize this track from a supplied animation and a notify track 
 	 * Adding all the FAnimNotifyEvents as sections in here and setting
 	 */
-	void Initialize(UAnimMontage& AnimationRef, const FAnimNotifyTrack& NotifyTrack);
+	void Initialize(UAnimSequenceBase& AnimationRef, const FAnimNotifyTrack& NotifyTrack);
 
 	/** Creates a new section for this track */
 	virtual UMovieSceneSection* CreateNewSection() override;
@@ -57,7 +57,7 @@ public:
 #endif
 
 	/** Returns a reference to the animation this track was created from */
-	UAnimMontage& GetAnimation() const;
+	UAnimSequenceBase& GetAnimation() const;
 
 private:
 
@@ -67,5 +67,5 @@ private:
 
 	/** Ptr to the animation this track was created from */
 	UPROPERTY()
-	TWeakObjectPtr<UAnimMontage> Animation;
+	TWeakObjectPtr<UAnimSequenceBase> Animation;
 };
\ No newline at end of file
diff --git a/Engine/Plugins/Experimental/Animation/ContextualAnimation/Source/ContextualAnimationEditor/Public/ContextualAnimViewModel.h b/Engine/Plugins/Experimental/Animation/ContextualAnimation/Source/ContextualAnimationEditor/Public/ContextualAnimViewModel.h
--- a/Engine/Plugins/Experimental/Animation/ContextualAnimation/Source/ContextualAnimationEditor/Public/ContextualAnimViewModel.h
+++ b/Engine/Plugins/Experimental/Animation/ContextualAnimation/Source/ContextualAnimationEditor/Public/ContextualAnimViewModel.h
@@ -20,7 +20,7 @@ class UMovieSceneTrack;
 class UMovieSceneSection;
 class UContextualAnimMovieSceneNotifyTrack;
 class UContextualAnimMovieSceneNotifySection;
-class UAnimMontage;
+class UAnimSequenceBase;
 class IDetailsView;
 class IStructureDetailsView;
 struct FMovieSceneSectionMovedParams;
@@ -51,9 +51,9 @@ public:
 
 	void AddNewVariant(const FContextualAnimNewVariantParams& Params);
 
- 	UAnimMontage* FindAnimationByGuid(const FGuid& Guid) const;
+	UAnimSequenceBase* FindAnimationByGuid(const FGuid& Guid) const;
 
-	void AnimationModified(UAnimMontage& Animation);
+	void AnimationModified(UAnimSequenceBase& Animation);
 
 	void SetActiveSceneVariantIdx(int32 Index);
 
@@ -102,7 +102,7 @@ private:
 	FContextualAnimStartSceneParams StartSceneParams;
 
 	/** Container for the animations on the time line. Should be removed once we add a proper animation track */
-	TArray<UAnimMontage*> AnimsBeingEdited;
+	TArray<UAnimSequenceBase*> AnimsBeingEdited;
 
 	AActor* SpawnPreviewActor(const FContextualAnimTrack& AnimTrack);
 
@@ -114,7 +114,7 @@ private:
 
 	void SequencerDataChanged(EMovieSceneDataChangeType DataChangeType);
 
-	void OnAnimNotifyChanged(UAnimMontage* Animation);
+	void OnAnimNotifyChanged(UAnimSequenceBase* Animation);
 
 	void CreateSequencer();
 };
\ No newline at end of file
