Removing EditorStyleSet.h includes

diff --git a/Engine/Plugins/Animation/ControlRig/Source/ControlRigEditor/Private/EditMode/ControlRigBaseDockableView.cpp b/Engine/Plugins/Animation/ControlRig/Source/ControlRigEditor/Private/EditMode/ControlRigBaseDockableView.cpp
--- a/Engine/Plugins/Animation/ControlRig/Source/ControlRigEditor/Private/EditMode/ControlRigBaseDockableView.cpp
+++ b/Engine/Plugins/Animation/ControlRig/Source/ControlRigEditor/Private/EditMode/ControlRigBaseDockableView.cpp
@@ -2,9 +2,7 @@
 
 #include "ControlRigBaseDockableView.h"
 #include "AssetRegistry/AssetData.h"
-#include "EditorStyleSet.h"
 #include "Widgets/Text/SInlineEditableTextBlock.h"
-#include "EditorStyleSet.h"
 #include "Styling/CoreStyle.h"
 #include "ScopedTransaction.h"
 #include "ControlRig.h"
diff --git a/Engine/Plugins/Animation/ControlRig/Source/ControlRigEditor/Private/EditMode/SControlRigDetails.cpp b/Engine/Plugins/Animation/ControlRig/Source/ControlRigEditor/Private/EditMode/SControlRigDetails.cpp
--- a/Engine/Plugins/Animation/ControlRig/Source/ControlRigEditor/Private/EditMode/SControlRigDetails.cpp
+++ b/Engine/Plugins/Animation/ControlRig/Source/ControlRigEditor/Private/EditMode/SControlRigDetails.cpp
@@ -5,9 +5,7 @@
 #include "Widgets/Notifications/SNotificationList.h"
 #include "Framework/MultiBox/MultiBoxBuilder.h"
 #include "AssetRegistry/AssetData.h"
-#include "EditorStyleSet.h"
 #include "Widgets/Text/SInlineEditableTextBlock.h"
-#include "EditorStyleSet.h"
 #include "Styling/CoreStyle.h"
 #include "ScopedTransaction.h"
 #include "ControlRig.h"
diff --git a/Engine/Plugins/Developer/Concert/ConcertApp/MultiUserServer/Source/MultiUserServer/Private/ConcertServerStyle.cpp b/Engine/Plugins/Developer/Concert/ConcertApp/MultiUserServer/Source/MultiUserServer/Private/ConcertServerStyle.cpp
--- a/Engine/Plugins/Developer/Concert/ConcertApp/MultiUserServer/Source/MultiUserServer/Private/ConcertServerStyle.cpp
+++ b/Engine/Plugins/Developer/Concert/ConcertApp/MultiUserServer/Source/MultiUserServer/Private/ConcertServerStyle.cpp
@@ -2,7 +2,6 @@
 
 #include "ConcertServerStyle.h"
 
-#include "EditorStyleSet.h"
 #include "Framework/Application/SlateApplication.h"
 #include "Interfaces/IPluginManager.h"
 #include "Styling/SlateStyleRegistry.h"
diff --git a/Engine/Plugins/Developer/Concert/ConcertApp/MultiUserServer/Source/MultiUserServer/Private/Widgets/SessionTabs/Archived/SConcertArchivedSessionInspector.cpp b/Engine/Plugins/Developer/Concert/ConcertApp/MultiUserServer/Source/MultiUserServer/Private/Widgets/SessionTabs/Archived/SConcertArchivedSessionInspector.cpp
--- a/Engine/Plugins/Developer/Concert/ConcertApp/MultiUserServer/Source/MultiUserServer/Private/Widgets/SessionTabs/Archived/SConcertArchivedSessionInspector.cpp
+++ b/Engine/Plugins/Developer/Concert/ConcertApp/MultiUserServer/Source/MultiUserServer/Private/Widgets/SessionTabs/Archived/SConcertArchivedSessionInspector.cpp
@@ -2,12 +2,12 @@
 
 #include "SConcertArchivedSessionInspector.h"
 
-#include "EditorStyleSet.h"
 #include "Framework/Docking/TabManager.h"
 #include "Session/History/SSessionHistory.h"
 #include "Widgets/Docking/SDockTab.h"
 #include "Widgets/Layout/SBorder.h"
 #include "Widgets/SBoxPanel.h"
+#include "Styling/AppStyle.h"
 
 #define LOCTEXT_NAMESPACE "SConcertArchivedSessionInspector"
 
diff --git a/Engine/Plugins/Editor/DisplayClusterLaunch/Source/DisplayClusterLaunchEditor/Private/DisplayClusterLaunchEditorStyle.cpp b/Engine/Plugins/Editor/DisplayClusterLaunch/Source/DisplayClusterLaunchEditor/Private/DisplayClusterLaunchEditorStyle.cpp
--- a/Engine/Plugins/Editor/DisplayClusterLaunch/Source/DisplayClusterLaunchEditor/Private/DisplayClusterLaunchEditorStyle.cpp
+++ b/Engine/Plugins/Editor/DisplayClusterLaunch/Source/DisplayClusterLaunchEditor/Private/DisplayClusterLaunchEditorStyle.cpp
@@ -2,7 +2,6 @@
 
 #include "DisplayClusterLaunchEditorStyle.h"
 
-#include "EditorStyleSet.h"
 #include "Framework/Application/SlateApplication.h"
 #include "Interfaces/IPluginManager.h"
 #include "Styling/CoreStyle.h"
diff --git a/Engine/Plugins/Experimental/Animation/ContextualAnimation/Source/ContextualAnimationEditor/Public/ContextualAnimAssetEditorCommands.h b/Engine/Plugins/Experimental/Animation/ContextualAnimation/Source/ContextualAnimationEditor/Public/ContextualAnimAssetEditorCommands.h
--- a/Engine/Plugins/Experimental/Animation/ContextualAnimation/Source/ContextualAnimationEditor/Public/ContextualAnimAssetEditorCommands.h
+++ b/Engine/Plugins/Experimental/Animation/ContextualAnimation/Source/ContextualAnimationEditor/Public/ContextualAnimAssetEditorCommands.h
@@ -10,7 +10,7 @@ class FContextualAnimAssetEditorCommands : public TCommands<FContextualAnimAsset
 {
 public:
 	FContextualAnimAssetEditorCommands()
-		: TCommands<FContextualAnimAssetEditorCommands>(TEXT("ContextualAnimAssetEditor"), NSLOCTEXT("Contexts", "ContextualAnim", "Contextual Anim"), NAME_None, FAppStyle::GetStyleSetName())
+		: TCommands<FContextualAnimAssetEditorCommands>(TEXT("ContextualAnimAssetEditor"), NSLOCTEXT("Contexts", "ContextualAnim", "Contextual Anim"), NAME_None, FAppStyle::GetAppStyleSetName())
 	{
 	}
 
diff --git a/Engine/Plugins/Experimental/Animation/PoseSearch/Source/Editor/Private/PoseSearchCustomization.cpp b/Engine/Plugins/Experimental/Animation/PoseSearch/Source/Editor/Private/PoseSearchCustomization.cpp
--- a/Engine/Plugins/Experimental/Animation/PoseSearch/Source/Editor/Private/PoseSearchCustomization.cpp
+++ b/Engine/Plugins/Experimental/Animation/PoseSearch/Source/Editor/Private/PoseSearchCustomization.cpp
@@ -11,7 +11,6 @@
 #include "ScopedTransaction.h"
 #include "Engine/EngineBaseTypes.h"
 #include "UObject/UnrealType.h"
-#include "EditorStyleSet.h"
 #include "Widgets/DeclarativeSyntaxSupport.h"
 #include "Widgets/SBoxPanel.h"
 #include "Widgets/Layout/SWrapBox.h"
diff --git a/Engine/Plugins/Experimental/CinematicPrestreaming/Source/CinematicPrestreamingEditor/Private/TrackEditors/CinePrestreamingTrackEditor.cpp b/Engine/Plugins/Experimental/CinematicPrestreaming/Source/CinematicPrestreamingEditor/Private/TrackEditors/CinePrestreamingTrackEditor.cpp
--- a/Engine/Plugins/Experimental/CinematicPrestreaming/Source/CinematicPrestreamingEditor/Private/TrackEditors/CinePrestreamingTrackEditor.cpp
+++ b/Engine/Plugins/Experimental/CinematicPrestreaming/Source/CinematicPrestreamingEditor/Private/TrackEditors/CinePrestreamingTrackEditor.cpp
@@ -3,7 +3,7 @@
 #include "TrackEditors/CinePrestreamingTrackEditor.h"
 
 #include "CinePrestreamingData.h"
-#include "EditorStyleSet.h"
+#include "Styling/AppStyle.h"
 #include "Framework/MultiBox/MultiBoxBuilder.h"
 #include "ISequencerSection.h"
 #include "LevelSequence.h"
diff --git a/Engine/Plugins/Experimental/MLDeformer/MLDeformerFramework/Source/MLDeformerFrameworkEditor/Private/MLDeformerEditorStyle.cpp b/Engine/Plugins/Experimental/MLDeformer/MLDeformerFramework/Source/MLDeformerFrameworkEditor/Private/MLDeformerEditorStyle.cpp
--- a/Engine/Plugins/Experimental/MLDeformer/MLDeformerFramework/Source/MLDeformerFrameworkEditor/Private/MLDeformerEditorStyle.cpp
+++ b/Engine/Plugins/Experimental/MLDeformer/MLDeformerFramework/Source/MLDeformerFrameworkEditor/Private/MLDeformerEditorStyle.cpp
@@ -1,7 +1,6 @@
 // Copyright Epic Games, Inc. All Rights Reserved.
 
 #include "MLDeformerEditorStyle.h"
-#include "EditorStyleSet.h"
 #include "Styling/SlateStyle.h"
 #include "Styling/SlateStyleRegistry.h"
 #include "Styling/SlateStyleMacros.h"
diff --git a/Engine/Plugins/Experimental/MLDeformer/MLDeformerFramework/Source/MLDeformerFrameworkEditor/Private/SMLDeformerCurvePickerWidget.cpp b/Engine/Plugins/Experimental/MLDeformer/MLDeformerFramework/Source/MLDeformerFrameworkEditor/Private/SMLDeformerCurvePickerWidget.cpp
--- a/Engine/Plugins/Experimental/MLDeformer/MLDeformerFramework/Source/MLDeformerFrameworkEditor/Private/SMLDeformerCurvePickerWidget.cpp
+++ b/Engine/Plugins/Experimental/MLDeformer/MLDeformerFramework/Source/MLDeformerFrameworkEditor/Private/SMLDeformerCurvePickerWidget.cpp
@@ -1,7 +1,6 @@
 // Copyright Epic Games, Inc. All Rights Reserved.
 
 #include "SMLDeformerCurvePickerWidget.h"
-#include "EditorStyleSet.h"
 #include "DetailLayoutBuilder.h"
 #include "Widgets/Input/SComboButton.h"
 #include "Widgets/Input/SSearchBox.h"
diff --git a/Engine/Plugins/Experimental/WidgetEditorToolPalette/Source/WidgetEditorToolPalette/Private/WidgetEditorToolPaletteCommands.cpp b/Engine/Plugins/Experimental/WidgetEditorToolPalette/Source/WidgetEditorToolPalette/Private/WidgetEditorToolPaletteCommands.cpp
--- a/Engine/Plugins/Experimental/WidgetEditorToolPalette/Source/WidgetEditorToolPalette/Private/WidgetEditorToolPaletteCommands.cpp
+++ b/Engine/Plugins/Experimental/WidgetEditorToolPalette/Source/WidgetEditorToolPalette/Private/WidgetEditorToolPaletteCommands.cpp
@@ -2,7 +2,6 @@
 
 #include "WidgetEditorToolPaletteCommands.h"
 #include "WidgetEditorToolPaletteStyle.h"
-#include "EditorStyleSet.h"
 #include "Framework/Commands/UICommandInfo.h"
 #include "Settings/CreateWidgetToolSettings.h"
 #include "Styling/SlateIconFinder.h"
diff --git a/Engine/Plugins/FX/Niagara/Source/NiagaraEditorWidgets/Private/SNiagaraParameterDropTarget.h b/Engine/Plugins/FX/Niagara/Source/NiagaraEditorWidgets/Private/SNiagaraParameterDropTarget.h
--- a/Engine/Plugins/FX/Niagara/Source/NiagaraEditorWidgets/Private/SNiagaraParameterDropTarget.h
+++ b/Engine/Plugins/FX/Niagara/Source/NiagaraEditorWidgets/Private/SNiagaraParameterDropTarget.h
@@ -10,7 +10,6 @@
 #include "Framework/SlateDelegates.h"
 #include "Widgets/DeclarativeSyntaxSupport.h"
 #include "Widgets/SCompoundWidget.h"
-#include "EditorStyleSet.h"
 #include "Styling/StyleColors.h"
 #include "ViewModels/Stack/NiagaraStackFunctionInput.h"
 
diff --git a/Engine/Plugins/Media/MediaPlate/Source/MediaPlateEditor/Private/AssetTools/MediaPlateActions.cpp b/Engine/Plugins/Media/MediaPlate/Source/MediaPlateEditor/Private/AssetTools/MediaPlateActions.cpp
--- a/Engine/Plugins/Media/MediaPlate/Source/MediaPlateEditor/Private/AssetTools/MediaPlateActions.cpp
+++ b/Engine/Plugins/Media/MediaPlate/Source/MediaPlateEditor/Private/AssetTools/MediaPlateActions.cpp
@@ -2,7 +2,6 @@
 
 #include "AssetTools/MediaPlateActions.h"
 #include "ContentBrowserModule.h"
-#include "EditorStyleSet.h"
 #include "Framework/MultiBox/MultiBoxBuilder.h"
 #include "IContentBrowserSingleton.h"
 #include "MediaPlateComponent.h"
diff --git a/Engine/Plugins/VirtualProduction/DMX/DMXPixelMapping/Source/DMXPixelMappingEditor/Private/Widgets/SDMXPixelMappingPaletteItem.cpp b/Engine/Plugins/VirtualProduction/DMX/DMXPixelMapping/Source/DMXPixelMappingEditor/Private/Widgets/SDMXPixelMappingPaletteItem.cpp
--- a/Engine/Plugins/VirtualProduction/DMX/DMXPixelMapping/Source/DMXPixelMappingEditor/Private/Widgets/SDMXPixelMappingPaletteItem.cpp
+++ b/Engine/Plugins/VirtualProduction/DMX/DMXPixelMapping/Source/DMXPixelMappingEditor/Private/Widgets/SDMXPixelMappingPaletteItem.cpp
@@ -6,7 +6,6 @@
 #include "Templates/DMXPixelMappingComponentTemplate.h"
 #include "ViewModels/DMXPixelMappingPaletteViewModel.h"
 
-#include "EditorStyleSet.h"
 #include "Widgets/Layout/SBorder.h"
 #include "Widgets/Text/STextBlock.h"
 #include "Widgets/Views/STableViewBase.h"
