Contextual Anim Editor (WIP): Refactored binding mechanism to support external data and extended blueprint interface

diff --git a/Engine/Plugins/Experimental/Animation/ContextualAnimation/Source/ContextualAnimation/Private/ContextualAnimManager.cpp b/Engine/Plugins/Experimental/Animation/ContextualAnimation/Source/ContextualAnimation/Private/ContextualAnimManager.cpp
--- a/Engine/Plugins/Experimental/Animation/ContextualAnimation/Source/ContextualAnimation/Private/ContextualAnimManager.cpp
+++ b/Engine/Plugins/Experimental/Animation/ContextualAnimation/Source/ContextualAnimation/Private/ContextualAnimManager.cpp
@@ -105,7 +105,7 @@ UContextualAnimSceneInstance* UContextualAnimManager::ForceStartScene(const UCon
 	{
 		FName RoleToBind = Pair.Key;
 
-		AActor* ActorToBind = Pair.Value;
+		AActor* ActorToBind = Pair.Value.GetActor();
 		if (ActorToBind == nullptr)
 		{
 			UE_LOG(LogContextualAnim, Warning, TEXT("UContextualAnimManager::ForceStartScene. Can't start scene. Reason: Trying to bind Invalid Actor. SceneAsset: %s Role: %s"),
@@ -123,9 +123,7 @@ UContextualAnimSceneInstance* UContextualAnimManager::ForceStartScene(const UCon
 			return nullptr;
 		}
 
-		check(AnimTrack);
-
-		Bindings.Add(FContextualAnimSceneActorData(RoleToBind, Params.VariantIdx, *ActorToBind, *AnimTrack, Params.AnimStartTime));
+		Bindings.Add(FContextualAnimSceneBinding(Pair.Value, SceneAsset, *AnimTrack));
 	}
 
 	UClass* Class = SceneAsset.GetSceneInstanceClass();
@@ -153,68 +151,38 @@ UContextualAnimSceneInstance* UContextualAnimManager::BP_TryStartScene(const UCo
 
 UContextualAnimSceneInstance* UContextualAnimManager::TryStartScene(const UContextualAnimSceneAsset& SceneAsset, const FContextualAnimStartSceneParams& Params)
 {
-	// Check that we have defined a primary role in the scene asset
-	const FName PrimaryRole = SceneAsset.GetPrimaryRole();
-	if (PrimaryRole.IsNone())
-	{
-		UE_LOG(LogContextualAnim, Warning, TEXT("UContextualAnimManager::TryStartScene. Can't start scene. Reason: Invalid Primary Role. SceneAsset: %s Role: %s"),
-			*GetNameSafe(&SceneAsset), *PrimaryRole.ToString());
-
-		return nullptr;
-	}
+	FContextualAnimSceneBindings Bindings;
 
-	// Find the actor that should be bound to the primary Role.
-	AActor* const* PrimaryActorPtr = Params.RoleToActorMap.Find(SceneAsset.GetPrimaryRole());
-	if (PrimaryActorPtr == nullptr)
+	bool bSuccess = false;
+	if (Params.VariantIdx != INDEX_NONE)
 	{
-		UE_LOG(LogContextualAnim, Warning, TEXT("UContextualAnimManager::TryStartScene. Can't start scene. Reason: Can't find valid actor for the Primary Role. SceneAsset: %s Role: %s"),
-			*GetNameSafe(&SceneAsset), *PrimaryRole.ToString());
-
-		return nullptr;
+		bSuccess = FContextualAnimSceneBindings::TryCreateBindings(SceneAsset, Params.VariantIdx, Params.RoleToActorMap, Bindings);
 	}
-
-	FContextualAnimSceneBindings Bindings;
-
-	FContextualAnimPrimaryActorData PrimaryActorData;
-	PrimaryActorData.Transform = (*PrimaryActorPtr)->GetActorTransform();
-
-	const int32 VariantsNum = SceneAsset.GetTotalVariants();
-	for (int32 VariantIdx = 0; VariantIdx < VariantsNum; VariantIdx++)
+	else
 	{
-		for (const auto& Pair : Params.RoleToActorMap)
+		const int32 VariantsNum = SceneAsset.GetTotalVariants();
+		for (int32 VariantIdx = 0; VariantIdx < VariantsNum; VariantIdx++)
 		{
-			FName RoleToBind = Pair.Key;
-			AActor* ActorToBind = Pair.Value;
-
-			FContextualAnimQuerierData QuerierData;
-			QuerierData.Transform = ActorToBind->GetActorTransform();
-			QuerierData.Velocity = ActorToBind->GetVelocity();
-
-			const FContextualAnimTrack* AnimTrack = SceneAsset.GetAnimTrack(RoleToBind, VariantIdx);
-			if (AnimTrack && AnimTrack->DoesQuerierPassSelectionCriteria(PrimaryActorData, QuerierData))
-			{
-				Bindings.Add(FContextualAnimSceneActorData(RoleToBind, VariantIdx, *ActorToBind, *AnimTrack));
-			}
-			else
+			if (FContextualAnimSceneBindings::TryCreateBindings(SceneAsset, VariantIdx, Params.RoleToActorMap, Bindings))
 			{
-				Bindings.Reset();
+				bSuccess = true;
 				break;
 			}
 		}
+	}
 
-		if (Params.RoleToActorMap.Num() == Bindings.Num())
-		{
-			UClass* Class = SceneAsset.GetSceneInstanceClass();
-			UContextualAnimSceneInstance* NewInstance = Class ? NewObject<UContextualAnimSceneInstance>(this, Class) : NewObject<UContextualAnimSceneInstance>(this);
-			NewInstance->SceneAsset = &SceneAsset;
-			NewInstance->Bindings = MoveTemp(Bindings);
-			NewInstance->Start();
-			NewInstance->OnSceneEnded.AddDynamic(this, &UContextualAnimManager::OnSceneInstanceEnded);
+	if (bSuccess)
+	{
+		UClass* Class = SceneAsset.GetSceneInstanceClass();
+		UContextualAnimSceneInstance* NewInstance = Class ? NewObject<UContextualAnimSceneInstance>(this, Class) : NewObject<UContextualAnimSceneInstance>(this);
+		NewInstance->SceneAsset = &SceneAsset;
+		NewInstance->Bindings = MoveTemp(Bindings);
+		NewInstance->Start();
+		NewInstance->OnSceneEnded.AddDynamic(this, &UContextualAnimManager::OnSceneInstanceEnded);
 
-			Instances.Add(NewInstance);
+		Instances.Add(NewInstance);
 
-			return NewInstance;
-		}
+		return NewInstance;
 	}
 
 	return nullptr;
diff --git a/Engine/Plugins/Experimental/Animation/ContextualAnimation/Source/ContextualAnimation/Private/ContextualAnimSceneActorComponent.cpp b/Engine/Plugins/Experimental/Animation/ContextualAnimation/Source/ContextualAnimation/Private/ContextualAnimSceneActorComponent.cpp
--- a/Engine/Plugins/Experimental/Animation/ContextualAnimation/Source/ContextualAnimation/Private/ContextualAnimSceneActorComponent.cpp
+++ b/Engine/Plugins/Experimental/Animation/ContextualAnimation/Source/ContextualAnimation/Private/ContextualAnimSceneActorComponent.cpp
@@ -56,13 +56,13 @@ void UContextualAnimSceneActorComponent::OnUnregister()
 	}
 }
 
-void UContextualAnimSceneActorComponent::OnJoinedScene(const FContextualAnimSceneActorData* SceneActorData)
+void UContextualAnimSceneActorComponent::OnJoinedScene(const FContextualAnimSceneBinding* Binding)
 {
-	check(SceneActorData);
-	SceneActorDataPtr = SceneActorData;
+	check(Binding);
+	BindingPtr = Binding;
 
 	// Start listening to TickPose if we joined an scene where we need IK
-	if (SceneActorData->GetIKTargetDefs().IKTargetDefs.Num() > 0)
+	if (Binding->GetIKTargetDefs().IKTargetDefs.Num() > 0)
 	{
 		if (USkeletalMeshComponent* SkelMeshComp = UContextualAnimUtilities::TryGetSkeletalMeshComponent(GetOwner()))
 		{
@@ -73,14 +73,14 @@ void UContextualAnimSceneActorComponent::OnJoinedScene(const FContextualAnimScen
 	OnJoinedSceneDelegate.Broadcast(this);
 }
 
-void UContextualAnimSceneActorComponent::OnLeftScene(const FContextualAnimSceneActorData* SceneActorData)
+void UContextualAnimSceneActorComponent::OnLeftScene(const FContextualAnimSceneBinding* Binding)
 {
-	check(SceneActorData);
+	check(Binding);
 
 	OnLeftSceneDelegate.Broadcast(this);
 
 	// Stop listening to TickPose if we were
-	if (SceneActorData->GetIKTargetDefs().IKTargetDefs.Num() > 0)
+	if (Binding->GetIKTargetDefs().IKTargetDefs.Num() > 0)
 	{
 		if (USkeletalMeshComponent* SkelMeshComp = UContextualAnimUtilities::TryGetSkeletalMeshComponent(GetOwner()))
 		{
@@ -88,7 +88,7 @@ void UContextualAnimSceneActorComponent::OnLeftScene(const FContextualAnimSceneA
 		}
 	}
 
-	SceneActorDataPtr = nullptr;
+	BindingPtr = nullptr;
 }
 
 void UContextualAnimSceneActorComponent::OnTickPose(class USkinnedMeshComponent* SkinnedMeshComponent, float DeltaTime, bool bNeedsValidRootMotion)
@@ -101,15 +101,15 @@ void UContextualAnimSceneActorComponent::UpdateIKTargets()
 {
 	IKTargets.Reset();
 
-	check(SceneActorDataPtr);
+	check(BindingPtr);
 
-	const FAnimMontageInstance* MontageInstance = SceneActorDataPtr->GetAnimMontageInstance();
+	const FAnimMontageInstance* MontageInstance = BindingPtr->GetAnimMontageInstance();
 	if(MontageInstance == nullptr)
 	{
 		return;
 	}
 
-	for (const FContextualAnimIKTargetDefinition& IKTargetDef : SceneActorDataPtr->GetIKTargetDefs().IKTargetDefs)
+	for (const FContextualAnimIKTargetDefinition& IKTargetDef : BindingPtr->GetIKTargetDefs().IKTargetDefs)
 	{
 #if !(UE_BUILD_SHIPPING || UE_BUILD_TEST)
 		const bool bDrawDebugEnable = CVarContextualAnimIKDebug.GetValueOnGameThread() > 0;
@@ -124,16 +124,16 @@ void UContextualAnimSceneActorComponent::UpdateIKTargets()
 		// @TODO: IKTargetTransform will be off by 1 frame if we tick before target. 
 		// Should we at least add an option to the SceneAsset to setup tick dependencies or should this be entirely up to the user?
 
-		if (const FContextualAnimSceneActorData* TargetSceneActorData = SceneActorDataPtr->GetSceneInstance().FindSceneActorDataByRole(IKTargetDef.TargetRoleName))
+		if (const FContextualAnimSceneBinding* TargetBinding = BindingPtr->GetSceneInstance()->FindBindingByRole(IKTargetDef.TargetRoleName))
 		{
-			if (const USkeletalMeshComponent* TargetSkelMeshComp = TargetSceneActorData->GetSkeletalMeshComponent())
+			if (const USkeletalMeshComponent* TargetSkelMeshComp = TargetBinding->GetSkeletalMeshComponent())
 			{
 				if (IKTargetDef.Provider == EContextualAnimIKTargetProvider::Autogenerated)
 				{
 					const FTransform IKTargetParentTransform = TargetSkelMeshComp->GetSocketTransform(IKTargetDef.TargetBoneName);
 
 					const float Time = MontageInstance->GetPosition();
-					IKTargetTransform = SceneActorDataPtr->GetAnimTrack().IKTargetData.ExtractTransformAtTime(IKTargetDef.GoalName, Time) * IKTargetParentTransform;
+					IKTargetTransform = BindingPtr->GetAnimTrack().IKTargetData.ExtractTransformAtTime(IKTargetDef.GoalName, Time) * IKTargetParentTransform;
 
 #if !(UE_BUILD_SHIPPING || UE_BUILD_TEST)
 					if (bDrawDebugEnable)
@@ -177,7 +177,7 @@ void UContextualAnimSceneActorComponent::UpdateIKTargets()
 #endif
 
 		// Convert IK Target to mesh space
-		//IKTargetTransform = IKTargetTransform.GetRelativeTransform(SceneActorDataPtr->GetSkeletalMeshComponent()->GetComponentTransform());
+		//IKTargetTransform = IKTargetTransform.GetRelativeTransform(BindingPtr->GetSkeletalMeshComponent()->GetComponentTransform());
 
 		IKTargets.Add(FContextualAnimIKTarget(IKTargetDef.GoalName, Alpha, IKTargetTransform));
 	}
diff --git a/Engine/Plugins/Experimental/Animation/ContextualAnimation/Source/ContextualAnimation/Private/ContextualAnimSceneAsset.cpp b/Engine/Plugins/Experimental/Animation/ContextualAnimation/Source/ContextualAnimation/Private/ContextualAnimSceneAsset.cpp
--- a/Engine/Plugins/Experimental/Animation/ContextualAnimation/Source/ContextualAnimation/Private/ContextualAnimSceneAsset.cpp
+++ b/Engine/Plugins/Experimental/Animation/ContextualAnimation/Source/ContextualAnimation/Private/ContextualAnimSceneAsset.cpp
@@ -468,7 +468,7 @@ TArray<FName> UContextualAnimSceneAsset::GetRoles() const
  	return Result;
 }
 
-const FContextualAnimTrack* UContextualAnimSceneAsset::FindFirstAnimTrackForRoleThatPassesSelectionCriteria(const FName& Role, const FContextualAnimPrimaryActorData& PrimaryActorData, const FContextualAnimQuerierData& QuerierData) const
+const FContextualAnimTrack* UContextualAnimSceneAsset::FindFirstAnimTrackForRoleThatPassesSelectionCriteria(const FName& Role, const FContextualAnimSceneBindingContext& Primary, const FContextualAnimSceneBindingContext& Querier) const
 {
 	for (const FContextualAnimTracksContainer& Variant : Variants)
 	{
@@ -479,7 +479,7 @@ const FContextualAnimTrack* UContextualAnimSceneAsset::FindFirstAnimTrackForRole
 				bool bSuccess = true;
 				for (const UContextualAnimSelectionCriterion* Criterion : AnimTrack.SelectionCriteria)
 				{
-					if (Criterion && !Criterion->DoesQuerierPassCondition(PrimaryActorData, QuerierData))
+					if (Criterion && !Criterion->DoesQuerierPassCondition(Primary, Querier))
 					{
 						bSuccess = false;
 						break;
@@ -497,7 +497,7 @@ const FContextualAnimTrack* UContextualAnimSceneAsset::FindFirstAnimTrackForRole
 	return nullptr;
 }
 
-const FContextualAnimTrack* UContextualAnimSceneAsset::FindAnimTrackForRoleWithClosestEntryLocation(const FName& Role, const FContextualAnimPrimaryActorData& PrimaryActorData, const FVector& TestLocation) const
+const FContextualAnimTrack* UContextualAnimSceneAsset::FindAnimTrackForRoleWithClosestEntryLocation(const FName& Role, const FContextualAnimSceneBindingContext& Primary, const FVector& TestLocation) const
 {
 	const FContextualAnimTrack* Result = nullptr;
 
@@ -508,7 +508,7 @@ const FContextualAnimTrack* UContextualAnimSceneAsset::FindAnimTrackForRoleWithC
 		{
 			if (AnimTrack.Role == Role)
 			{
-				const FTransform EntryTransform = AnimTrack.GetAlignmentTransformAtEntryTime() * PrimaryActorData.Transform;
+				const FTransform EntryTransform = AnimTrack.GetAlignmentTransformAtEntryTime() * Primary.GetTransform();
 				const float DistSq = FVector::DistSquared(EntryTransform.GetLocation(), TestLocation);
 				if (DistSq < BestDistanceSq)
 				{
@@ -526,28 +526,25 @@ bool UContextualAnimSceneAsset::Query(FName Role, FContextualAnimQueryResult& Ou
 {
 	//@TODO: Kept around only to do not break existing content. It will go away in the future.
 
-	FContextualAnimPrimaryActorData PrimaryActorData;
-	PrimaryActorData.Transform = ToWorldTransform;
+	FContextualAnimSceneBindingContext Primary(ToWorldTransform);
+	FContextualAnimSceneBindingContext Querier(QueryParams.Querier.IsValid() ? QueryParams.Querier->GetActorTransform() : QueryParams.QueryTransform);
 
-	FContextualAnimQuerierData QuerierData;
-	QuerierData.Transform = QueryParams.Querier.IsValid() ? QueryParams.Querier->GetActorTransform() : QueryParams.QueryTransform;
-
-	if(const FContextualAnimTrack* AnimTrack = FindFirstAnimTrackForRoleThatPassesSelectionCriteria(Role, PrimaryActorData, QuerierData))
+	if (const FContextualAnimTrack* AnimTrack = FindFirstAnimTrackForRoleThatPassesSelectionCriteria(Role, Primary, Querier))
 	{
 		OutResult.VariantIdx = AnimTrack->VariantIdx;
-		OutResult.Animation = Cast<UAnimMontage>(AnimTrack->Animation); // Its ok to cast here to montage. Content using this function uses montage
+		OutResult.Animation = Cast<UAnimMontage>(AnimTrack->Animation);
 		OutResult.EntryTransform = AnimTrack->GetAlignmentTransformAtEntryTime() * ToWorldTransform;
 		OutResult.SyncTransform = AnimTrack->GetAlignmentTransformAtSyncTime() * ToWorldTransform;
 
 		if (QueryParams.bFindAnimStartTime)
 		{
-			const FVector LocalLocation = (QuerierData.Transform.GetRelativeTransform(ToWorldTransform)).GetLocation();
+			const FVector LocalLocation = (Querier.GetTransform().GetRelativeTransform(ToWorldTransform)).GetLocation();
 			OutResult.AnimStartTime = AnimTrack->FindBestAnimStartTime(LocalLocation);
 		}
 
 		return true;
 	}
-	
+
 	return false;
 }
 
diff --git a/Engine/Plugins/Experimental/Animation/ContextualAnimation/Source/ContextualAnimation/Private/ContextualAnimSceneInstance.cpp b/Engine/Plugins/Experimental/Animation/ContextualAnimation/Source/ContextualAnimation/Private/ContextualAnimSceneInstance.cpp
--- a/Engine/Plugins/Experimental/Animation/ContextualAnimation/Source/ContextualAnimation/Private/ContextualAnimSceneInstance.cpp
+++ b/Engine/Plugins/Experimental/Animation/ContextualAnimation/Source/ContextualAnimation/Private/ContextualAnimSceneInstance.cpp
@@ -14,77 +14,7 @@
 #include "Engine/World.h"
 #include "DrawDebugHelpers.h"
 
-UContextualAnimSceneActorComponent* FContextualAnimSceneActorData::GetSceneActorComponent() const
-{
-	//@TODO: Cache this during the binding
-	return (Actor.IsValid()) ? Actor->FindComponentByClass<UContextualAnimSceneActorComponent>() : nullptr;
-}
-
-FTransform FContextualAnimSceneActorData::GetTransform() const
-{
-	const UContextualAnimSceneActorComponent* Comp = GetSceneActorComponent();
-	return Comp ? Comp->GetComponentTransform() : Actor->GetActorTransform();
-}
-
-UAnimInstance* FContextualAnimSceneActorData::GetAnimInstance() const
-{
-	return UContextualAnimUtilities::TryGetAnimInstance(GetActor());
-}
-
-USkeletalMeshComponent* FContextualAnimSceneActorData::GetSkeletalMeshComponent() const
-{
-	return UContextualAnimUtilities::TryGetSkeletalMeshComponent(GetActor());
-}
-
-FAnimMontageInstance* FContextualAnimSceneActorData::GetAnimMontageInstance() const
-{
-	if (UAnimInstance* AnimInstance = GetAnimInstance())
-	{
-		return AnimInstance->GetActiveMontageInstance();
-	}
-
-	return nullptr;
-}
-
-float FContextualAnimSceneActorData::GetAnimTime() const
-{
-	const FAnimMontageInstance* MontageInstance = GetAnimMontageInstance();
-	return MontageInstance ? MontageInstance->GetPosition() : -1.f;
-}
-
-FName FContextualAnimSceneActorData::GetCurrentSection() const
-{
-	const FAnimMontageInstance* MontageInstance = GetAnimMontageInstance();
-	return MontageInstance ? MontageInstance->GetCurrentSection() : NAME_None;
-}
-
-int32 FContextualAnimSceneActorData::GetCurrentSectionIndex() const
-{
-	if(const FAnimMontageInstance* MontageInstance = GetAnimMontageInstance())
-	{
-		float CurrentPosition;
-		return MontageInstance->Montage->GetAnimCompositeSectionIndexFromPos(MontageInstance->GetPosition(), CurrentPosition);
-	}
-
-	return INDEX_NONE;
-}
-
-const FContextualAnimIKTargetDefContainer& FContextualAnimSceneActorData::GetIKTargetDefs() const
-{ 
-	return GetSceneInstance().GetSceneAsset().GetIKTargetDefsForRole(Role); 
-}
-
-//================================================================================================================
-
-void UContextualAnimSceneInstance::BreakContextualAnimSceneActorData(const FContextualAnimSceneActorData& SceneActorData, AActor*& Actor, UAnimMontage*& Montage, float& AnimTime, int32& CurrentSectionIndex, FName& CurrentSectionName)
-{
-	Actor = SceneActorData.GetActor();
-	Montage = Cast<UAnimMontage>(SceneActorData.GetAnimTrack().Animation);
-	AnimTime = SceneActorData.GetAnimTime();
-	CurrentSectionIndex = SceneActorData.GetCurrentSectionIndex();
-	CurrentSectionName = SceneActorData.GetCurrentSection();
-}
-
+// UContextualAnimSceneInstance
 //================================================================================================================
 
 UContextualAnimSceneInstance::UContextualAnimSceneInstance(const FObjectInitializer& ObjectInitializer)
@@ -104,31 +34,31 @@ void UContextualAnimSceneInstance::Tick(float DeltaTime)
 
 bool UContextualAnimSceneInstance::IsActorInThisScene(const AActor* Actor) const
 {
-	return FindSceneActorDataByActor(Actor) != nullptr;
+	return FindBindingByActor(Actor) != nullptr;
 }
 
 AActor* UContextualAnimSceneInstance::GetActorByRole(FName Role) const
 {
-	const FContextualAnimSceneActorData* SceneActorData = FindSceneActorDataByRole(Role);
-	return SceneActorData ? SceneActorData->GetActor() : nullptr;
+	const FContextualAnimSceneBinding* Binding = FindBindingByRole(Role);
+	return Binding ? Binding->GetActor() : nullptr;
 }
 
-void UContextualAnimSceneInstance::Join(FContextualAnimSceneActorData& SceneActorData)
+void UContextualAnimSceneInstance::Join(FContextualAnimSceneBinding& Binding)
 {
-	AActor* Actor = SceneActorData.GetActor();
+	AActor* Actor = Binding.GetActor();
 	if (Actor == nullptr)
 	{
 		return;
 	}
 
-	if (UAnimSequenceBase* Animation = SceneActorData.GetAnimTrack().Animation)
+	if (UAnimSequenceBase* Animation = Binding.GetAnimTrack().Animation)
 	{
-		if (UAnimInstance* AnimInstance = SceneActorData.GetAnimInstance())
+		if (UAnimInstance* AnimInstance = Binding.GetAnimInstance())
 		{
 			// Keep montage support for now but might go away soon
 			if (UAnimMontage* AnimMontage = Cast<UAnimMontage>(Animation))
 			{
-				AnimInstance->Montage_Play(AnimMontage, 1.f, EMontagePlayReturnType::MontageLength, SceneActorData.GetAnimStartTime());
+				AnimInstance->Montage_Play(AnimMontage, 1.f, EMontagePlayReturnType::MontageLength);
 			}
 			else
 			{
@@ -139,8 +69,8 @@ void UContextualAnimSceneInstance::Join(FContextualAnimSceneActorData& SceneActo
 				const float InPlayRate = 1.f;
 				const int32 LoopCount = 1;
 				const float BlendOutTriggerTime = -1.f;
-				const float InTimeToStartMontageAt = SceneActorData.GetAnimStartTime();
-				AnimInstance->PlaySlotAnimationAsDynamicMontage(SceneActorData.GetAnimTrack().Animation, SlotName, BlendInTime, BlendOutTime, InPlayRate, LoopCount, BlendOutTriggerTime, InTimeToStartMontageAt);
+				const float InTimeToStartMontageAt = 0.f;
+				AnimInstance->PlaySlotAnimationAsDynamicMontage(Binding.GetAnimTrack().Animation, SlotName, BlendInTime, BlendOutTime, InPlayRate, LoopCount, BlendOutTriggerTime, InTimeToStartMontageAt);
 			}
 
 			AnimInstance->OnPlayMontageNotifyBegin.AddUniqueDynamic(this, &UContextualAnimSceneInstance::OnNotifyBeginReceived);
@@ -150,7 +80,7 @@ void UContextualAnimSceneInstance::Join(FContextualAnimSceneActorData& SceneActo
 
 
 		//@TODO: Temp, until we have a way to switch between movement mode using AnimNotifyState
-		if (SceneActorData.GetAnimTrack().bRequireFlyingMode)
+		if (Binding.GetAnimTrack().bRequireFlyingMode)
 		{
 			if (UCharacterMovementComponent* CharacterMovementComp = Actor->FindComponentByClass<UCharacterMovementComponent>())
 			{
@@ -165,8 +95,8 @@ void UContextualAnimSceneInstance::Join(FContextualAnimSceneActorData& SceneActo
 		{
 			const FName WarpTargetName = Pair.Key;
 			const FTransform ScenePivotRuntime = Pair.Value;
-			const float Time = SceneActorData.GetAnimTrack().GetSyncTimeForWarpSection(WarpTargetName);
-			const FTransform TransformRelativeToScenePivot = SceneActorData.GetAnimTrack().AlignmentData.ExtractTransformAtTime(WarpTargetName, Time);
+			const float Time = Binding.GetAnimTrack().GetSyncTimeForWarpSection(WarpTargetName);
+			const FTransform TransformRelativeToScenePivot = Binding.GetAnimTrack().AlignmentData.ExtractTransformAtTime(WarpTargetName, Time);
 			const FTransform WarpTarget = (TransformRelativeToScenePivot * ScenePivotRuntime);
 
 			MotionWarpComp->AddOrUpdateWarpTargetFromTransform(WarpTargetName, WarpTarget);
@@ -178,24 +108,24 @@ void UContextualAnimSceneInstance::Join(FContextualAnimSceneActorData& SceneActo
 		SetIgnoreCollisionWithOtherActors(Actor, true);
 	}
 
-	SceneActorData.SceneInstancePtr = this;
+	Binding.SceneInstancePtr = this;
 
-	if(UContextualAnimSceneActorComponent* SceneActorComp = SceneActorData.GetSceneActorComponent())
+	if(UContextualAnimSceneActorComponent* SceneActorComp = Binding.GetSceneActorComponent())
 	{
-		SceneActorComp->OnJoinedScene(&SceneActorData);
+		SceneActorComp->OnJoinedScene(&Binding);
 	}
 
 	OnActorJoined.Broadcast(this, Actor);
 }
 
-void UContextualAnimSceneInstance::Leave(FContextualAnimSceneActorData& SceneActorData)
+void UContextualAnimSceneInstance::Leave(FContextualAnimSceneBinding& Binding)
 {
 	// Check if we have an exit section and transition to it, otherwise just stop the montage
 	static const FName ExitSectionName = FName(TEXT("Exit"));
-	if(TransitionTo(SceneActorData, ExitSectionName) == false)
+	if(TransitionTo(Binding, ExitSectionName) == false)
 	{
-		AActor* Actor = SceneActorData.GetActor();
-		if (UAnimInstance* AnimInstance = SceneActorData.GetAnimInstance())
+		AActor* Actor = Binding.GetActor();
+		if (UAnimInstance* AnimInstance = Binding.GetAnimInstance())
 		{
 			UAnimMontage* CurrentMontage = AnimInstance->GetCurrentActiveMontage();
 			if (CurrentMontage)
@@ -206,9 +136,9 @@ void UContextualAnimSceneInstance::Leave(FContextualAnimSceneActorData& SceneAct
 	}
 }
 
-bool UContextualAnimSceneInstance::TransitionTo(FContextualAnimSceneActorData& SceneActorData, const FName& ToSectionName)
+bool UContextualAnimSceneInstance::TransitionTo(FContextualAnimSceneBinding& Binding, const FName& ToSectionName)
 {
-	UAnimInstance* AnimInstance = SceneActorData.GetAnimInstance();
+	UAnimInstance* AnimInstance = Binding.GetAnimInstance();
 	if (AnimInstance == nullptr)
 	{		
 		return false;
@@ -218,7 +148,7 @@ bool UContextualAnimSceneInstance::TransitionTo(FContextualAnimSceneActorData& S
 	if (CurrentMontage == nullptr)
 	{
 		UE_LOG(LogContextualAnim, Log, TEXT("UContextualAnimSceneInstance::TransitionTo. Actor is not playing any montage. Actor: %s ToSectionName: %s"),
-			*GetNameSafe(SceneActorData.GetActor()), *ToSectionName.ToString());
+			*GetNameSafe(Binding.GetActor()), *ToSectionName.ToString());
 
 		return false;
 	}
@@ -227,13 +157,13 @@ bool UContextualAnimSceneInstance::TransitionTo(FContextualAnimSceneActorData& S
 	if (SectionIdx == INDEX_NONE)
 	{
 		UE_LOG(LogContextualAnim, Log, TEXT("UContextualAnimSceneInstance::TransitionTo. Invalid Section. Actor: %s CurrentMontage: %s ToSectionName: %s"),
-			*GetNameSafe(SceneActorData.GetActor()), *GetNameSafe(CurrentMontage), *ToSectionName.ToString());
+			*GetNameSafe(Binding.GetActor()), *GetNameSafe(CurrentMontage), *ToSectionName.ToString());
 
 		return false;
 	}
 
 	UE_LOG(LogContextualAnim, Verbose, TEXT("UContextualAnimSceneInstance::TransitionTo. Actor: %s CurrentMontage: %s ToSectionName: %s"),
-		*GetNameSafe(SceneActorData.GetActor()), *GetNameSafe(CurrentMontage), *ToSectionName.ToString());
+		*GetNameSafe(Binding.GetActor()), *GetNameSafe(CurrentMontage), *ToSectionName.ToString());
 
 	// Unbind blend out delegate for a moment so we don't get it during the transition
 	AnimInstance->OnMontageBlendingOut.RemoveDynamic(this, &UContextualAnimSceneInstance::OnMontageBlendingOut);
@@ -251,14 +181,14 @@ void UContextualAnimSceneInstance::Start()
 	for (const FContextualAnimAlignmentSectionData AligmentData : SceneAsset->GetAlignmentSections())
 	{
 		FTransform ScenePivotRuntime = FTransform::Identity;
-		if (const FContextualAnimSceneActorData* SceneActor = FindSceneActorDataByRole(AligmentData.Origin))
+		if (const FContextualAnimSceneBinding* Binding = FindBindingByRole(AligmentData.Origin))
 		{
 			if (AligmentData.bAlongClosestDistance)
 			{
-				if (const FContextualAnimSceneActorData* OtherSceneActor = FindSceneActorDataByRole(AligmentData.OtherRole))
+				if (const FContextualAnimSceneBinding* OtherBinding = FindBindingByRole(AligmentData.OtherRole))
 				{
-					const FTransform T1 = SceneActor->GetActor()->GetActorTransform();
-					const FTransform T2 = OtherSceneActor->GetActor()->GetActorTransform();
+					const FTransform T1 = Binding->GetTransform();
+					const FTransform T2 = OtherBinding->GetTransform();
 
 					ScenePivotRuntime.SetLocation(FMath::Lerp<FVector>(T1.GetLocation(), T2.GetLocation(), AligmentData.Weight));
 					ScenePivotRuntime.SetRotation((T2.GetLocation() - T1.GetLocation()).GetSafeNormal2D().ToOrientationQuat());
@@ -266,7 +196,7 @@ void UContextualAnimSceneInstance::Start()
 			}
 			else
 			{
-				ScenePivotRuntime = SceneActor->GetActor()->GetActorTransform();
+				ScenePivotRuntime = Binding->GetTransform();
 			}
 		}
 
diff --git a/Engine/Plugins/Experimental/Animation/ContextualAnimation/Source/ContextualAnimation/Private/ContextualAnimSelectionCriterion.cpp b/Engine/Plugins/Experimental/Animation/ContextualAnimation/Source/ContextualAnimation/Private/ContextualAnimSelectionCriterion.cpp
--- a/Engine/Plugins/Experimental/Animation/ContextualAnimation/Source/ContextualAnimation/Private/ContextualAnimSelectionCriterion.cpp
+++ b/Engine/Plugins/Experimental/Animation/ContextualAnimation/Source/ContextualAnimation/Private/ContextualAnimSelectionCriterion.cpp
@@ -30,24 +30,27 @@ UContextualAnimSelectionCriterion_TriggerArea::UContextualAnimSelectionCriterion
 	PolygonPoints.Add(FVector(100, 100, 0));
 }
 
-bool UContextualAnimSelectionCriterion_TriggerArea::DoesQuerierPassCondition(const FContextualAnimPrimaryActorData& PrimaryActorData, const FContextualAnimQuerierData& QuerierData) const
+bool UContextualAnimSelectionCriterion_TriggerArea::DoesQuerierPassCondition(const FContextualAnimSceneBindingContext& Primary, const FContextualAnimSceneBindingContext& Querier) const
 {
 	if (PolygonPoints.Num() == 4)
 	{
+		const FTransform PrimaryTransform = Primary.GetTransform();
+		const FTransform QuerierTransform = Querier.GetTransform();
+
 		const float HalfHeight = FMath::Max((Height / 2.f), 0.f);
-		const float VDist = FMath::Abs((PrimaryActorData.Transform.GetLocation().Z + PolygonPoints[0].Z + HalfHeight) - QuerierData.Transform.GetLocation().Z);
+		const float VDist = FMath::Abs((PrimaryTransform.GetLocation().Z + PolygonPoints[0].Z + HalfHeight) - QuerierTransform.GetLocation().Z);
 		if (VDist > HalfHeight)
 		{
 			return false;
 		}
 
-		const FVector2D TestPoint = FVector2D(QuerierData.Transform.GetLocation());
+		const FVector2D TestPoint = FVector2D(QuerierTransform.GetLocation());
 		const int32 NumPoints = PolygonPoints.Num();
 		float AngleSum = 0.0f;
 		for (int32 PointIndex = 0; PointIndex < NumPoints; ++PointIndex)
 		{
-			const FVector2D& VecAB = FVector2D(PrimaryActorData.Transform.TransformPositionNoScale(PolygonPoints[PointIndex])) - TestPoint;
-			const FVector2D& VecAC = FVector2D(PrimaryActorData.Transform.TransformPositionNoScale(PolygonPoints[(PointIndex + 1) % NumPoints])) - TestPoint;
+			const FVector2D& VecAB = FVector2D(PrimaryTransform.TransformPositionNoScale(PolygonPoints[PointIndex])) - TestPoint;
+			const FVector2D& VecAC = FVector2D(PrimaryTransform.TransformPositionNoScale(PolygonPoints[(PointIndex + 1) % NumPoints])) - TestPoint;
 			const float Angle = FMath::Sign(FVector2D::CrossProduct(VecAB, VecAC)) * FMath::Acos(FMath::Clamp(FVector2D::DotProduct(VecAB, VecAC) / (VecAB.Size() * VecAC.Size()), -1.0f, 1.0f));
 			AngleSum += Angle;
 		}
@@ -61,13 +64,16 @@ bool UContextualAnimSelectionCriterion_TriggerArea::DoesQuerierPassCondition(con
 // UContextualAnimSelectionCriterion_Facing
 //===========================================================================
 
-bool UContextualAnimSelectionCriterion_Facing::DoesQuerierPassCondition(const FContextualAnimPrimaryActorData& PrimaryActorData, const FContextualAnimQuerierData& QuerierData) const
+bool UContextualAnimSelectionCriterion_Facing::DoesQuerierPassCondition(const FContextualAnimSceneBindingContext& Primary, const FContextualAnimSceneBindingContext& Querier) const
 {
 	if (MaxAngle > 0.f)
 	{
+		const FTransform PrimaryTransform = Primary.GetTransform();
+		const FTransform QuerierTransform = Querier.GetTransform();
+
 		const float FacingCos = FMath::Cos(FMath::Clamp(FMath::DegreesToRadians(MaxAngle), 0.f, PI));
-		const FVector ToTarget = (PrimaryActorData.Transform.GetLocation() - QuerierData.Transform.GetLocation()).GetSafeNormal2D();
-		const float DotProduct = FVector::DotProduct(QuerierData.Transform.GetRotation().GetForwardVector(), ToTarget);
+		const FVector ToTarget = (PrimaryTransform.GetLocation() - QuerierTransform.GetLocation()).GetSafeNormal2D();
+		const float DotProduct = FVector::DotProduct(QuerierTransform.GetRotation().GetForwardVector(), ToTarget);
 		return (DotProduct > FacingCos);
 	}
 
diff --git a/Engine/Plugins/Experimental/Animation/ContextualAnimation/Source/ContextualAnimation/Private/ContextualAnimTypes.cpp b/Engine/Plugins/Experimental/Animation/ContextualAnimation/Source/ContextualAnimation/Private/ContextualAnimTypes.cpp
--- a/Engine/Plugins/Experimental/Animation/ContextualAnimation/Source/ContextualAnimation/Private/ContextualAnimTypes.cpp
+++ b/Engine/Plugins/Experimental/Animation/ContextualAnimation/Source/ContextualAnimation/Private/ContextualAnimTypes.cpp
@@ -2,15 +2,19 @@
 
 #include "ContextualAnimTypes.h"
 #include "AnimationUtils.h"
+#include "Animation/AnimInstance.h"
 #include "Animation/AnimMontage.h"
 #include "Containers/ArrayView.h"
 #include "AnimNotifyState_MotionWarping.h"
 #include "ContextualAnimUtilities.h"
 #include "RootMotionModifier.h"
 #include "ContextualAnimSelectionCriterion.h"
+#include "ContextualAnimSceneActorComponent.h"
+#include "ContextualAnimSceneAsset.h"
 
 DEFINE_LOG_CATEGORY(LogContextualAnim);
 
+const FContextualAnimSceneBinding FContextualAnimSceneBinding::InvalidBinding;
 const FContextualAnimTrack FContextualAnimTrack::EmptyTrack;
 const FContextualAnimIKTarget FContextualAnimIKTarget::InvalidIKTarget;
 const FContextualAnimIKTargetDefContainer FContextualAnimIKTargetDefContainer::EmptyContainer;
@@ -171,15 +175,182 @@ float FContextualAnimTrack::FindBestAnimStartTime(const FVector& LocalLocation)
 	return BestTime;
 }
 
-bool FContextualAnimTrack::DoesQuerierPassSelectionCriteria(const FContextualAnimPrimaryActorData& PrimaryActorData, const FContextualAnimQuerierData& QuerierData) const
+bool FContextualAnimTrack::DoesQuerierPassSelectionCriteria(const FContextualAnimSceneBindingContext& Primary, const FContextualAnimSceneBindingContext& Querier) const
 {
 	for (const UContextualAnimSelectionCriterion* Criterion : SelectionCriteria)
 	{
-		if (Criterion && !Criterion->DoesQuerierPassCondition(PrimaryActorData, QuerierData))
+		if (Criterion && !Criterion->DoesQuerierPassCondition(Primary, Querier))
 		{
 			return false;
 		}
 	}
 
 	return true;
-}
\ No newline at end of file
+}
+
+// FContextualAnimSceneBindingContext
+///////////////////////////////////////////////////////////////////////
+
+void FContextualAnimSceneBindingContext::SetExternalTransform(const FTransform& InTransform)
+{
+	ExternalTransform = InTransform;
+}
+
+FTransform FContextualAnimSceneBindingContext::GetTransform() const
+{
+	if (ExternalTransform.IsSet())
+	{
+		return ExternalTransform.GetValue();
+	}
+	else if (AActor* ActorPtr = GetActor())
+	{
+		return ActorPtr->GetActorTransform();
+	}
+
+	return FTransform::Identity;
+}
+
+FVector FContextualAnimSceneBindingContext::GetVelocity() const
+{
+	if (ExternalVelocity.IsSet())
+	{
+		return ExternalVelocity.GetValue();
+	}
+	else if (AActor* ActorPtr = GetActor())
+	{
+		return ActorPtr->GetVelocity();
+	}
+
+	return FVector::ZeroVector;
+}
+
+// FContextualAnimSceneBinding
+///////////////////////////////////////////////////////////////////////
+
+FContextualAnimSceneBinding::FContextualAnimSceneBinding(const FContextualAnimSceneBindingContext& InContext, const UContextualAnimSceneAsset& InSceneAsset, const FContextualAnimTrack& InAnimTrack)
+{
+	check(InSceneAsset.HasValidData());
+
+	Context = InContext;
+	SceneAsset = &InSceneAsset;
+	AnimTrackPtr = &InAnimTrack;
+	RoleDefPtr = InSceneAsset.GetRolesAsset()->FindRoleDefinitionByName(InAnimTrack.Role);
+	check(RoleDefPtr);
+}
+
+UContextualAnimSceneActorComponent* FContextualAnimSceneBinding::GetSceneActorComponent() const
+{
+	//@TODO: Cache this during the binding
+	AActor* Actor = Context.GetActor();
+	return Actor ? Actor->FindComponentByClass<UContextualAnimSceneActorComponent>() : nullptr;
+}
+
+UAnimInstance* FContextualAnimSceneBinding::GetAnimInstance() const
+{
+	return UContextualAnimUtilities::TryGetAnimInstance(GetActor());
+}
+
+USkeletalMeshComponent* FContextualAnimSceneBinding::GetSkeletalMeshComponent() const
+{
+	return UContextualAnimUtilities::TryGetSkeletalMeshComponent(GetActor());
+}
+
+FAnimMontageInstance* FContextualAnimSceneBinding::GetAnimMontageInstance() const
+{
+	if (UAnimInstance* AnimInstance = GetAnimInstance())
+	{
+		return AnimInstance->GetActiveMontageInstance();
+	}
+
+	return nullptr;
+}
+
+float FContextualAnimSceneBinding::GetAnimMontageTime() const
+{
+	const FAnimMontageInstance* MontageInstance = GetAnimMontageInstance();
+	return MontageInstance ? MontageInstance->GetPosition() : -1.f;
+}
+
+FName FContextualAnimSceneBinding::GetCurrentSection() const
+{
+	const FAnimMontageInstance* MontageInstance = GetAnimMontageInstance();
+	return MontageInstance ? MontageInstance->GetCurrentSection() : NAME_None;
+}
+
+int32 FContextualAnimSceneBinding::GetCurrentSectionIndex() const
+{
+	if (const FAnimMontageInstance* MontageInstance = GetAnimMontageInstance())
+	{
+		float CurrentPosition;
+		return MontageInstance->Montage->GetAnimCompositeSectionIndexFromPos(MontageInstance->GetPosition(), CurrentPosition);
+	}
+
+	return INDEX_NONE;
+}
+
+const FContextualAnimIKTargetDefContainer& FContextualAnimSceneBinding::GetIKTargetDefs() const
+{
+	return GetSceneAsset().GetIKTargetDefsForRole(GetAnimTrack().Role);
+}
+
+// FContextualAnimSceneBindings
+///////////////////////////////////////////////////////////////////////
+
+bool FContextualAnimSceneBindings::TryCreateBindings(const UContextualAnimSceneAsset& SceneAsset, int32 VariantIdx, const TMap<FName, FContextualAnimSceneBindingContext>& Params, FContextualAnimSceneBindings& OutBindings)
+{
+	check(SceneAsset.HasValidData());
+
+	OutBindings.Reset();
+
+	// Find the actor that should be bound to the primary Role.
+	const FName PrimaryRole = SceneAsset.GetPrimaryRole();
+	const FContextualAnimSceneBindingContext* PrimaryPtr = Params.Find(PrimaryRole);
+	if (PrimaryPtr == nullptr)
+	{
+		UE_LOG(LogContextualAnim, Warning, TEXT("FContextualAnimSceneBindings::TryCreateBindings Failed. Reason: Can't find valid actor for primary role. SceneAsset: %s PrimaryRole: %s"),
+			*GetNameSafe(&SceneAsset), *PrimaryRole.ToString());
+
+		return false;
+	}
+
+	// First, try to bind primary track. 
+	// @TODO: Revisit this, passing the same data twice (as primary and querier) feels weird, but this allow us to run the selection mechanism even on the primary actor.
+
+	const FContextualAnimTrack* PrimaryAnimTrack = SceneAsset.GetAnimTrack(PrimaryRole, VariantIdx);
+	if (PrimaryAnimTrack && PrimaryAnimTrack->DoesQuerierPassSelectionCriteria(*PrimaryPtr, *PrimaryPtr))
+	{
+		OutBindings.Add(FContextualAnimSceneBinding(*PrimaryPtr, SceneAsset, *PrimaryAnimTrack));
+	}
+	else
+	{
+		UE_LOG(LogContextualAnim, Warning, TEXT("FContextualAnimSceneBindings::TryCreateBindings Failed. Reason: Can't find valid track for primary actor. SceneAsset: %s Role: %s Actor: %s Variant: %d"),
+			*GetNameSafe(&SceneAsset), *PrimaryRole.ToString(), *GetNameSafe(PrimaryPtr->GetActor()), VariantIdx);
+
+		return false;
+	}
+
+	// Now try to bind secondary tracks
+
+	for (const auto& Pair : Params)
+	{
+		FName RoleToBind = Pair.Key;
+		if (RoleToBind != PrimaryRole)
+		{
+			const FContextualAnimTrack* AnimTrack = SceneAsset.GetAnimTrack(RoleToBind, VariantIdx);
+			if (AnimTrack && AnimTrack->DoesQuerierPassSelectionCriteria(*PrimaryPtr, Pair.Value))
+			{
+				OutBindings.Add(FContextualAnimSceneBinding(Pair.Value, SceneAsset, *AnimTrack));
+			}
+			else
+			{
+				UE_LOG(LogContextualAnim, Warning, TEXT("FContextualAnimSceneBindings::TryCreateBindings Failed. Reason: Can't find valid track for secondary actor. SceneAsset: %s Role: %s Actor: %s Variant: %d"),
+					*GetNameSafe(&SceneAsset), *RoleToBind.ToString(), *GetNameSafe(Pair.Value.GetActor()), VariantIdx);
+
+				return false;
+			}
+		}
+	}
+
+	// Success if all the roles were filled
+	return (OutBindings.Num() == SceneAsset.GetNumRoles());
+}
diff --git a/Engine/Plugins/Experimental/Animation/ContextualAnimation/Source/ContextualAnimation/Private/ContextualAnimUtilities.cpp b/Engine/Plugins/Experimental/Animation/ContextualAnimation/Source/ContextualAnimation/Private/ContextualAnimUtilities.cpp
--- a/Engine/Plugins/Experimental/Animation/ContextualAnimation/Source/ContextualAnimation/Private/ContextualAnimUtilities.cpp
+++ b/Engine/Plugins/Experimental/Animation/ContextualAnimation/Source/ContextualAnimation/Private/ContextualAnimUtilities.cpp
@@ -19,6 +19,7 @@
 #include "EngineUtils.h"
 #include "AnimNotifyState_IKWindow.h"
 #include "SceneManagement.h"
+#include "MotionWarpingComponent.h"
 
 void UContextualAnimUtilities::ExtractLocalSpacePose(const UAnimSequenceBase* Animation, const FBoneContainer& BoneContainer, float Time, bool bExtractRootMotion, FCompactPose& OutPose)
 {
@@ -256,4 +257,104 @@ void UContextualAnimUtilities::DrawSector(FPrimitiveDrawInterface& PDI, const FV
 		PDI.DrawLine(LineStart, LineEnd, Color, DepthPriority, Thickness);
 		LastDirection = NewDirection;
 	}
-}
\ No newline at end of file
+}
+
+bool UContextualAnimUtilities::BP_CreateContextualAnimSceneBindings(const UContextualAnimSceneAsset* SceneAsset, const TMap<FName, FContextualAnimSceneBindingContext>& Params, FContextualAnimSceneBindings& OutBindings)
+{
+	if(SceneAsset == nullptr || !SceneAsset->HasValidData())
+	{
+		UE_LOG(LogContextualAnim, Warning, TEXT("UContextualAnimUtilities::FindBestVariantForActors Failed. Reason: Invalid or Empty SceneAsset (%s)"), *GetNameSafe(SceneAsset));
+		return false;
+	}
+
+	for (int32 VariantIdx = 0; VariantIdx < SceneAsset->GetTotalVariants(); VariantIdx++)
+	{
+		OutBindings.Reset();
+		if(FContextualAnimSceneBindings::TryCreateBindings(*SceneAsset, VariantIdx, Params, OutBindings))
+		{
+			return true;
+		}
+	}
+	
+	return false;
+}
+
+bool UContextualAnimUtilities::CalculateScenePivotForAlignmentSection(const FContextualAnimAlignmentSectionData& AligmentSectionData, const FContextualAnimSceneBindings& Bindings, FTransform& OutScenePivot)
+{
+	if (const FContextualAnimSceneBinding* Binding = Bindings.FindBindingByRole(AligmentSectionData.Origin))
+	{
+		if (AligmentSectionData.bAlongClosestDistance)
+		{
+			if (const FContextualAnimSceneBinding* OtherBinding = Bindings.FindBindingByRole(AligmentSectionData.OtherRole))
+			{
+				const FTransform T1 = Binding->GetTransform();
+				const FTransform T2 = OtherBinding->GetTransform();
+
+				OutScenePivot.SetLocation(FMath::Lerp<FVector>(T1.GetLocation(), T2.GetLocation(), AligmentSectionData.Weight));
+				OutScenePivot.SetRotation((T2.GetLocation() - T1.GetLocation()).GetSafeNormal2D().ToOrientationQuat());
+				return true;
+			}
+		}
+		else
+		{
+			OutScenePivot = Binding->GetTransform();
+			return true;
+		}
+	}
+
+	return false;
+}
+
+// SceneBindings Blueprint Interface
+//------------------------------------------------------------------------------------------
+
+const FContextualAnimSceneBinding& UContextualAnimUtilities::BP_SceneBindings_GetBindingByRole(const FContextualAnimSceneBindings& Bindings, FName Role)
+{
+	if(const FContextualAnimSceneBinding* SceneActorData = Bindings.FindBindingByRole(Role))
+	{
+		return *SceneActorData;
+	}
+
+	return FContextualAnimSceneBinding::InvalidBinding;
+}
+
+void UContextualAnimUtilities::BP_SceneBindings_AddOrUpdateWarpTargetsForBindings(const FContextualAnimSceneBindings& Bindings)
+{
+	const UContextualAnimSceneAsset* SceneAsset = Bindings.GetSceneAsset();
+	if(ensureAlways(SceneAsset))
+	{
+		for (const FContextualAnimAlignmentSectionData& AlignmentSection : SceneAsset->GetAlignmentSections())
+		{
+			FTransform ScenePivot = FTransform::Identity;
+			if (UContextualAnimUtilities::CalculateScenePivotForAlignmentSection(AlignmentSection, Bindings, ScenePivot))
+			{
+				for (const FContextualAnimSceneBinding& Binding : Bindings)
+				{
+					//@TODO: Cache this
+					const float Time = Binding.GetAnimTrack().GetSyncTimeForWarpSection(AlignmentSection.WarpTargetName);
+
+					const FTransform TransformRelativeToScenePivot = Binding.GetAnimTrack().AlignmentData.ExtractTransformAtTime(AlignmentSection.WarpTargetName, Time);
+					const FTransform WarpTarget = (TransformRelativeToScenePivot * ScenePivot);
+
+					if (UMotionWarpingComponent* MotionWarpComp = Binding.GetActor()->FindComponentByClass<UMotionWarpingComponent>())
+					{
+						MotionWarpComp->AddOrUpdateWarpTargetFromTransform(AlignmentSection.WarpTargetName, WarpTarget);
+					}
+				}
+			}
+		}
+	}
+}
+
+FTransform UContextualAnimUtilities::BP_SceneBindings_GetAlignmentTransformForRoleRelativeToOtherRole(const FContextualAnimSceneBindings& Bindings, FName Role, FName RelativeToRole, float Time = 0.f)
+{
+	FTransform Result = FTransform::Identity;
+
+	if(const UContextualAnimSceneAsset* SceneAsset = Bindings.GetSceneAsset())
+	{
+		Result = SceneAsset->GetAlignmentTransformForRoleRelativeToOtherRole(Role, RelativeToRole, Bindings.GetVariantIdx(), Time);
+	}
+
+	return Result;
+}
+
diff --git a/Engine/Plugins/Experimental/Animation/ContextualAnimation/Source/ContextualAnimation/Public/ContextualAnimSceneActorComponent.h b/Engine/Plugins/Experimental/Animation/ContextualAnimation/Source/ContextualAnimation/Public/ContextualAnimSceneActorComponent.h
--- a/Engine/Plugins/Experimental/Animation/ContextualAnimation/Source/ContextualAnimation/Public/ContextualAnimSceneActorComponent.h
+++ b/Engine/Plugins/Experimental/Animation/ContextualAnimation/Source/ContextualAnimation/Public/ContextualAnimSceneActorComponent.h
@@ -11,7 +11,7 @@ class UAnimInstance;
 class UAnimMontage;
 class AActor;
 class UContextualAnimSceneInstance;
-struct FContextualAnimSceneActorData;
+struct FContextualAnimSceneBinding;
 
 DECLARE_DYNAMIC_MULTICAST_DELEGATE_OneParam(FContextualAnimSceneActorCompDelegate, class UContextualAnimSceneActorComponent*, SceneActorComponent);
 
@@ -46,10 +46,10 @@ public:
 	virtual void AddIKGoals_Implementation(TMap<FName, FIKRigGoal>& OutGoals) override;
 
 	/** Called from the scene instance when the actor owner of this component joins an scene */
-	void OnJoinedScene(const FContextualAnimSceneActorData* SceneActorData);
+	void OnJoinedScene(const FContextualAnimSceneBinding* Binding);
 	
 	/** Called from the scene instance when the actor owner of this component leave an scene */
-	void OnLeftScene(const FContextualAnimSceneActorData* SceneActorData);
+	void OnLeftScene(const FContextualAnimSceneBinding* Binding);
 
 	UFUNCTION(BlueprintPure, Category = "Contextual Anim|Scene Actor Component")
 	const TArray<FContextualAnimIKTarget>& GetIKTargets() const { return IKTargets; }
@@ -59,8 +59,8 @@ public:
 
 protected:
 
-	/** Ptr back to the scene actor data that represent us in the scene instance we are part of */
-	const FContextualAnimSceneActorData* SceneActorDataPtr = nullptr;
+	/** Ptr back to the binding that represent us in the scene instance we are part of */
+	const FContextualAnimSceneBinding* BindingPtr = nullptr;
 
 	/** List of IKTarget for this frame */
 	UPROPERTY()
diff --git a/Engine/Plugins/Experimental/Animation/ContextualAnimation/Source/ContextualAnimation/Public/ContextualAnimSceneAsset.h b/Engine/Plugins/Experimental/Animation/ContextualAnimation/Source/ContextualAnimation/Public/ContextualAnimSceneAsset.h
--- a/Engine/Plugins/Experimental/Animation/ContextualAnimation/Source/ContextualAnimation/Public/ContextualAnimSceneAsset.h
+++ b/Engine/Plugins/Experimental/Animation/ContextualAnimation/Source/ContextualAnimation/Public/ContextualAnimSceneAsset.h
@@ -10,42 +10,6 @@
 class UContextualAnimScenePivotProvider;
 class UContextualAnimSceneInstance;
 
-USTRUCT(BlueprintType)
-struct FContextualAnimAlignmentSectionData
-{
-	GENERATED_BODY()
-
-	UPROPERTY(EditAnywhere, BlueprintReadWrite, Category = "Defaults")
-	FName WarpTargetName = NAME_None;
-
-	UPROPERTY(EditAnywhere, BlueprintReadWrite, Category = "Defaults", meta = (GetOptions = "GetRoles"))
-	FName Origin = NAME_None;
-		
-	UPROPERTY(EditAnywhere, BlueprintReadWrite, Category = "Defaults")
-	bool bAlongClosestDistance = false;
-
-	UPROPERTY(EditAnywhere, BlueprintReadWrite, Category = "Defaults", meta = (GetOptions = "GetRoles", EditCondition = "bAlongClosestDistance"))
-	FName OtherRole = NAME_None;
-	
-	UPROPERTY(EditAnywhere, BlueprintReadWrite, Category = "Defaults", meta = (ClampMin = "0.0", ClampMax = "1.0", UIMin = "0.0", UIMax = "1.0", EditCondition = "bAlongClosestDistance"))
-	float Weight = 0.f;
-};
-
-USTRUCT(BlueprintType)
-struct FContextualAnimRoleDefinition
-{
-	GENERATED_BODY()
-
-	UPROPERTY(EditAnywhere, BlueprintReadWrite, Category = "Defaults")
-	FName Name;
-
-	UPROPERTY(EditAnywhere, BlueprintReadWrite, Category = "Defaults")
-	TSubclassOf<AActor> PreviewActorClass;
-
-	UPROPERTY(EditAnywhere, BlueprintReadWrite, Category = "Defaults")
-	FTransform MeshToComponent = FTransform(FRotator(0.f, -90.f, 0.f));
-};
-
 UCLASS(Blueprintable)
 class CONTEXTUALANIMATION_API UContextualAnimRolesAsset : public UDataAsset
 {
@@ -62,6 +26,8 @@ public:
 	{
 		return Roles.FindByPredicate([Name](const FContextualAnimRoleDefinition& RoleDef) { return RoleDef.Name == Name; });
 	}
+
+	FORCEINLINE int32 GetNumRoles() const { return Roles.Num(); }
 };
 
 USTRUCT(BlueprintType)
@@ -105,9 +71,15 @@ public:
 	FORCEINLINE const TSubclassOf<UContextualAnimSceneInstance>& GetSceneInstanceClass() const { return SceneInstanceClass; }
 	FORCEINLINE const TArray<FContextualAnimAlignmentSectionData>& GetAlignmentSections() const { return  AlignmentSections; }
 
+	bool HasValidData() const { return RolesAsset && Variants.Num() > 0; }
+
+	const UContextualAnimRolesAsset* GetRolesAsset() const { return RolesAsset; }
+
 	UFUNCTION()
 	TArray<FName> GetRoles() const;
 
+	int32 GetNumRoles() const { return RolesAsset ? RolesAsset->GetNumRoles() : 0; }
+
 	const FContextualAnimTrack* GetAnimTrack(const FName& Role, int32 VariantIdx) const;
 
 	UFUNCTION(BlueprintCallable, Category = "Contextual Anim|Scene Asset")
@@ -115,9 +87,9 @@ public:
 
 	FName FindRoleByAnimation(const UAnimSequenceBase* Animation) const;
 
-	const FContextualAnimTrack* FindFirstAnimTrackForRoleThatPassesSelectionCriteria(const FName& Role, const FContextualAnimPrimaryActorData& PrimaryActorData, const FContextualAnimQuerierData& QuerierData) const;
-	
-	const FContextualAnimTrack* FindAnimTrackForRoleWithClosestEntryLocation(const FName& Role, const FContextualAnimPrimaryActorData& PrimaryActorData, const FVector& TestLocation) const;
+	const FContextualAnimTrack* FindFirstAnimTrackForRoleThatPassesSelectionCriteria(const FName& Role, const FContextualAnimSceneBindingContext& Primary, const FContextualAnimSceneBindingContext& Querier) const;
+
+	const FContextualAnimTrack* FindAnimTrackForRoleWithClosestEntryLocation(const FName& Role, const FContextualAnimSceneBindingContext& Primary, const FVector& TestLocation) const;
 
 	UFUNCTION(BlueprintCallable, Category = "Contextual Anim|Scene Asset")
 	FTransform GetAlignmentTransformForRoleRelativeToScenePivot(FName Role, int32 VariantIdx, float Time) const;
diff --git a/Engine/Plugins/Experimental/Animation/ContextualAnimation/Source/ContextualAnimation/Public/ContextualAnimSceneInstance.h b/Engine/Plugins/Experimental/Animation/ContextualAnimation/Source/ContextualAnimation/Public/ContextualAnimSceneInstance.h
--- a/Engine/Plugins/Experimental/Animation/ContextualAnimation/Source/ContextualAnimation/Public/ContextualAnimSceneInstance.h
+++ b/Engine/Plugins/Experimental/Animation/ContextualAnimation/Source/ContextualAnimation/Public/ContextualAnimSceneInstance.h
@@ -19,112 +19,6 @@ class USkeletalMeshComponent;
 class UAnimMontage;
 class UWorld;
 
-/** Represent an actor bound to a role in the scene */
-USTRUCT(BlueprintType)
-struct CONTEXTUALANIMATION_API FContextualAnimSceneActorData
-{
-	GENERATED_BODY()
-
-	FContextualAnimSceneActorData(){}
-	FContextualAnimSceneActorData(const FName& InRole, int32 InVariantIdx, AActor& InActor, const FContextualAnimTrack& InAnimTrack, float InAnimStartTime = 0.f)
-		: Role(InRole), VariantIdx(InVariantIdx), Actor(&InActor), AnimTrackPtr(&InAnimTrack), AnimStartTime(InAnimStartTime)
-	{}
-
-	/** Return a pointer to the actual actor in the world */
-	FORCEINLINE AActor* GetActor() const { return Actor.Get(); }
-	FORCEINLINE FName GetRole() const { return Role; }
-	FORCEINLINE int32 GetVariantIdx() const { return VariantIdx; }
-	FORCEINLINE float GetAnimStartTime() const { return AnimStartTime; }
-	FORCEINLINE const FContextualAnimTrack& GetAnimTrack() const { return *AnimTrackPtr; }
-	FORCEINLINE const UContextualAnimSceneInstance& GetSceneInstance() const { return *SceneInstancePtr.Get(); }
-
-	const FContextualAnimIKTargetDefContainer& GetIKTargetDefs() const;
-
-	/** Return the transform used for alignment for this scene actor */
-	FTransform GetTransform() const;
-
-	/** Return the current playback time of the animation this actor is playing */
-	float GetAnimTime() const;
-
-	FName GetCurrentSection() const;
-
-	int32 GetCurrentSectionIndex() const;
-
-	/** Returns the ActiveMontageInstance or null in the case of static actors */
-	FAnimMontageInstance* GetAnimMontageInstance() const;
-
-	UAnimInstance* GetAnimInstance() const;
-
-	USkeletalMeshComponent* GetSkeletalMeshComponent() const;
-
-	UContextualAnimSceneActorComponent* GetSceneActorComponent() const;
-
-private:
-
-	friend UContextualAnimSceneInstance;
-
-	/** Role this actor is representing */
-	FName Role = NAME_None;
-
-	int32 VariantIdx = INDEX_NONE;
-
-	/** The actual actor in the world */
-	TWeakObjectPtr<AActor> Actor = nullptr;
-
-	/** Ptr to the animation data in the scene asset used by this actor */
-	const FContextualAnimTrack* AnimTrackPtr = nullptr;
-
-	/** Desired time to start the animation */
-	float AnimStartTime = 0.f;
-
-	/** Ptr back to the scene instance we belong to */
-	TWeakObjectPtr<const UContextualAnimSceneInstance> SceneInstancePtr = nullptr;
-
-#if WITH_EDITOR
-public:
-	/** Guid only used in editor to bind this actor to sequencer */
-	FGuid Guid;
-#endif
-
-};
-
-USTRUCT(BlueprintType)
-struct CONTEXTUALANIMATION_API FContextualAnimSceneBindings
-{
-	GENERATED_BODY()
-	
-	const FContextualAnimSceneActorData* FindSceneActorDataByActor(const AActor* Actor) const
-	{
-		return Actor ? Data.FindByPredicate([Actor](const auto& Item) { return Item.GetActor() == Actor; }) : nullptr;
-	}
-
-	const FContextualAnimSceneActorData* FindSceneActorDataByRole(const FName& Role) const
-	{
-		return Role != NAME_None ? Data.FindByPredicate([&Role](const auto& Item) { return Item.GetRole() == Role; }) : nullptr;
-	}
-
-#if WITH_EDITOR
-	const FContextualAnimSceneActorData* FindSceneActorDataByGuid(const FGuid& Guid) const
-	{
-		return Guid.IsValid() ? Data.FindByPredicate([&Guid](const auto& Item) { return Item.Guid == Guid; }) : nullptr;
-	}
-#endif
-
-	FORCEINLINE int32 Num() const { return Data.Num(); }
-	FORCEINLINE int32 Add(const FContextualAnimSceneActorData& NewData) { return Data.Add(NewData); }
-	FORCEINLINE void Reset() { return Data.Reset(); }
-
-	FORCEINLINE TArray<FContextualAnimSceneActorData>::RangedForIteratorType      begin() { return Data.begin(); }
-	FORCEINLINE TArray<FContextualAnimSceneActorData>::RangedForConstIteratorType begin() const { return Data.begin(); }
-	FORCEINLINE TArray<FContextualAnimSceneActorData>::RangedForIteratorType      end() { return Data.end(); }
-	FORCEINLINE TArray<FContextualAnimSceneActorData>::RangedForConstIteratorType end() const { return Data.end(); }
-
-private:
-
-	UPROPERTY()
-	TArray<FContextualAnimSceneActorData> Data;
-};
-
 /** Delegate to notify external objects when this is scene is completed */
 DECLARE_DYNAMIC_MULTICAST_DELEGATE_OneParam(FOnContextualAnimSceneEnded, class UContextualAnimSceneInstance*, SceneInstance);
 
@@ -186,29 +80,25 @@ public:
 	const UContextualAnimSceneAsset& GetSceneAsset() const { return *SceneAsset; }
 	const FContextualAnimSceneBindings& GetBindings() const { return Bindings; };
 	FContextualAnimSceneBindings& GetBindings() { return Bindings; };
-	const FContextualAnimSceneActorData* FindSceneActorDataByActor(const AActor* Actor) const { return Bindings.FindSceneActorDataByActor(Actor); }
-	const FContextualAnimSceneActorData* FindSceneActorDataByRole(const FName& Role) const { return Bindings.FindSceneActorDataByRole(Role); }
+	const FContextualAnimSceneBinding* FindBindingByActor(const AActor* Actor) const { return Bindings.FindBindingByActor(Actor); }
+	const FContextualAnimSceneBinding* FindBindingByRole(const FName& Role) const { return Bindings.FindBindingByRole(Role); }
 
 #if WITH_EDITOR
-	const FContextualAnimSceneActorData* FindSceneActorDataByGuid(const FGuid& Guid) const { return Bindings.FindSceneActorDataByGuid(Guid); }
+	const FContextualAnimSceneBinding* FindBindingByGuid(const FGuid& Guid) const { return Bindings.FindBindingByGuid(Guid); }
 #endif
 
 	UFUNCTION(BlueprintCallable, Category = "Contextual Anim|Scene Instance")
 	AActor* GetActorByRole(FName Role) const;
 
-	/** Extracts data from a ContextualAnimSceneActorData */
-	UFUNCTION(BlueprintPure, Category = "Contextual Anim|Scene Actor Data", meta = (NativeBreakFunc))
-	static void BreakContextualAnimSceneActorData(const FContextualAnimSceneActorData& SceneActorData, AActor*& Actor, UAnimMontage*& Montage, float& AnimTime, int32& CurrentSectionIndex, FName& CurrentSectionName);
-
 protected:
 
 	/** Tells the scene actor to join the scene (play animation) */
-	void Join(FContextualAnimSceneActorData& SceneActorData);
+	void Join(FContextualAnimSceneBinding& Binding);
 
 	/** Tells the scene actor to leave the scene (stop animation) */
-	void Leave(FContextualAnimSceneActorData& SceneActorData);
+	void Leave(FContextualAnimSceneBinding& Binding);
 
-	bool TransitionTo(FContextualAnimSceneActorData& SceneActorData, const FName& ToSectionName);
+	bool TransitionTo(FContextualAnimSceneBinding& Binding, const FName& ToSectionName);
 
 	/** Helper function to set ignore collision between the supplied actor and all the other actors in this scene */
 	void SetIgnoreCollisionWithOtherActors(AActor* Actor, bool bValue) const;
diff --git a/Engine/Plugins/Experimental/Animation/ContextualAnimation/Source/ContextualAnimation/Public/ContextualAnimSelectionCriterion.h b/Engine/Plugins/Experimental/Animation/ContextualAnimation/Source/ContextualAnimation/Public/ContextualAnimSelectionCriterion.h
--- a/Engine/Plugins/Experimental/Animation/ContextualAnimation/Source/ContextualAnimation/Public/ContextualAnimSelectionCriterion.h
+++ b/Engine/Plugins/Experimental/Animation/ContextualAnimation/Source/ContextualAnimation/Public/ContextualAnimSelectionCriterion.h
@@ -21,7 +21,7 @@ public:
 
 	class UContextualAnimSceneAsset* GetSceneAssetOwner() const;
 
-	virtual bool DoesQuerierPassCondition(const FContextualAnimPrimaryActorData& PrimaryActorData, const FContextualAnimQuerierData& QuerierData) const { return false; }
+	virtual bool DoesQuerierPassCondition(const FContextualAnimSceneBindingContext& Primary, const FContextualAnimSceneBindingContext& Querier) const { return false; }
 };
 
 // UContextualAnimSelectionCriterion_TriggerArea
@@ -42,7 +42,7 @@ public:
 
 	UContextualAnimSelectionCriterion_TriggerArea(const FObjectInitializer& ObjectInitializer);
 
-	virtual bool DoesQuerierPassCondition(const FContextualAnimPrimaryActorData& PrimaryActorData, const FContextualAnimQuerierData& QuerierData) const override;
+	virtual bool DoesQuerierPassCondition(const FContextualAnimSceneBindingContext& Primary, const FContextualAnimSceneBindingContext& Querier) const override;
 };
 
 // UContextualAnimSelectionCriterion_Facing
@@ -60,5 +60,5 @@ public:
 
 	UContextualAnimSelectionCriterion_Facing(const FObjectInitializer& ObjectInitializer) : Super(ObjectInitializer) {}
 
-	virtual bool DoesQuerierPassCondition(const FContextualAnimPrimaryActorData& PrimaryActorData, const FContextualAnimQuerierData& QuerierData) const override;
+	virtual bool DoesQuerierPassCondition(const FContextualAnimSceneBindingContext& Primary, const FContextualAnimSceneBindingContext& Querier) const override;
 };
\ No newline at end of file
diff --git a/Engine/Plugins/Experimental/Animation/ContextualAnimation/Source/ContextualAnimation/Public/ContextualAnimTypes.h b/Engine/Plugins/Experimental/Animation/ContextualAnimation/Source/ContextualAnimation/Public/ContextualAnimTypes.h
--- a/Engine/Plugins/Experimental/Animation/ContextualAnimation/Source/ContextualAnimation/Public/ContextualAnimTypes.h
+++ b/Engine/Plugins/Experimental/Animation/ContextualAnimation/Source/ContextualAnimation/Public/ContextualAnimTypes.h
@@ -12,6 +12,10 @@ CONTEXTUALANIMATION_API DECLARE_LOG_CATEGORY_EXTERN(LogContextualAnim, Log, All)
 class AActor;
 class UAnimSequenceBase;
 class UContextualAnimSelectionCriterion;
+class UContextualAnimSceneAsset;
+class UContextualAnimSceneInstance;
+class UContextualAnimSceneActorComponent;
+struct FAnimMontageInstance;
 
 namespace UE 
 {
@@ -98,7 +102,7 @@ struct CONTEXTUALANIMATION_API FContextualAnimTrack
 
 	float FindBestAnimStartTime(const FVector& LocalLocation) const;
 
-	bool DoesQuerierPassSelectionCriteria(const FContextualAnimPrimaryActorData& PrimaryActorData, const FContextualAnimQuerierData& QuerierData) const;
+	bool DoesQuerierPassSelectionCriteria(const FContextualAnimSceneBindingContext& PrimaryActorData, const FContextualAnimSceneBindingContext& QuerierData) const;
 
 	static const FContextualAnimTrack EmptyTrack;
 };
@@ -188,6 +192,201 @@ struct CONTEXTUALANIMATION_API FContextualAnimIKTargetDefContainer
 	static const FContextualAnimIKTargetDefContainer EmptyContainer;
 };
 
+// FContextualAnimRoleDefinition
+///////////////////////////////////////////////////////////////////////
+
+USTRUCT(BlueprintType)
+struct FContextualAnimRoleDefinition
+{
+	GENERATED_BODY()
+
+	UPROPERTY(EditAnywhere, BlueprintReadWrite, Category = "Defaults")
+	FName Name;
+
+	UPROPERTY(EditAnywhere, BlueprintReadWrite, Category = "Defaults")
+	TSubclassOf<AActor> PreviewActorClass;
+
+	UPROPERTY(EditAnywhere, BlueprintReadWrite, Category = "Defaults")
+	FTransform MeshToComponent = FTransform(FRotator(0.f, -90.f, 0.f));
+};
+
+// FContextualAnimAlignmentSectionData
+///////////////////////////////////////////////////////////////////////
+
+USTRUCT(BlueprintType)
+struct FContextualAnimAlignmentSectionData
+{
+	GENERATED_BODY()
+
+	UPROPERTY(EditAnywhere, BlueprintReadWrite, Category = "Defaults")
+	FName WarpTargetName = NAME_None;
+
+	UPROPERTY(EditAnywhere, BlueprintReadWrite, Category = "Defaults", meta = (GetOptions = "GetRoles"))
+	FName Origin = NAME_None;
+
+	UPROPERTY(EditAnywhere, BlueprintReadWrite, Category = "Defaults")
+	bool bAlongClosestDistance = false;
+
+	UPROPERTY(EditAnywhere, BlueprintReadWrite, Category = "Defaults", meta = (GetOptions = "GetRoles", EditCondition = "bAlongClosestDistance"))
+	FName OtherRole = NAME_None;
+
+	UPROPERTY(EditAnywhere, BlueprintReadWrite, Category = "Defaults", meta = (ClampMin = "0.0", ClampMax = "1.0", UIMin = "0.0", UIMax = "1.0", EditCondition = "bAlongClosestDistance"))
+	float Weight = 0.f;
+};
+
+// FContextualAnimSceneBindingContext
+///////////////////////////////////////////////////////////////////////
+
+USTRUCT(BlueprintType, meta = (HasNativeMake = "ContextualAnimUtilities.BP_SceneBindingContext_MakeFromActor"))
+struct CONTEXTUALANIMATION_API FContextualAnimSceneBindingContext
+{
+	GENERATED_BODY()
+
+	FContextualAnimSceneBindingContext() {}
+
+	FContextualAnimSceneBindingContext(AActor* InActor, const TOptional<FTransform>& InExternalTransform = TOptional<FTransform>(), const TOptional<FVector>& InExternalVelocity = TOptional<FVector>())
+		: Actor(InActor), ExternalTransform(InExternalTransform), ExternalVelocity(InExternalVelocity) {}
+
+	FContextualAnimSceneBindingContext(const FTransform& InExternalTransform, const TOptional<FVector>& InExternalVelocity = TOptional<FVector>())
+		: ExternalTransform(InExternalTransform), ExternalVelocity(InExternalVelocity) {}
+
+	AActor* GetActor() const { return Actor.Get(); }
+
+	void SetExternalTransform(const FTransform& InTransform);
+
+	FTransform GetTransform() const;
+
+	FVector GetVelocity() const;
+
+	//@TODO: Add accessors for GameplayTags
+
+private:
+
+	TWeakObjectPtr<AActor> Actor = nullptr;
+
+	TOptional<FTransform> ExternalTransform;
+
+	TOptional<FVector> ExternalVelocity;
+
+	FGameplayTagContainer ExternalGameplayTags;
+};
+
+/** Represent an actor bound to a role in the scene */
+USTRUCT(BlueprintType)
+struct CONTEXTUALANIMATION_API FContextualAnimSceneBinding
+{
+	GENERATED_BODY()
+
+	FContextualAnimSceneBinding() {}
+	FContextualAnimSceneBinding(const FContextualAnimSceneBindingContext& InContext, const UContextualAnimSceneAsset& InSceneAsset, const FContextualAnimTrack& InAnimTrack);
+
+	FORCEINLINE const FContextualAnimSceneBindingContext& GetContext() const { return Context; }
+	FORCEINLINE FContextualAnimSceneBindingContext& GetContext() { return Context; }
+	FORCEINLINE AActor* GetActor() const { return GetContext().GetActor(); }
+	FORCEINLINE FTransform GetTransform() const { return GetContext().GetTransform(); }
+	FORCEINLINE FVector GetVelocity()  const { return GetContext().GetVelocity(); }
+
+	/** Returns the track in the scene asset used by this actor */
+	FORCEINLINE const FContextualAnimTrack& GetAnimTrack() const { check(AnimTrackPtr);  return *AnimTrackPtr; }
+
+	/** Returns the role definition used by this actor */
+	FORCEINLINE const FContextualAnimRoleDefinition& GetRoleDef() const { check(RoleDefPtr); return *RoleDefPtr; }
+
+	/** Returns the SceneInstance we belong to (if any) */
+	FORCEINLINE const UContextualAnimSceneInstance* GetSceneInstance() const { return SceneInstancePtr.Get(); }
+
+	/** Returns the SceneAsset we were created from */
+	FORCEINLINE const UContextualAnimSceneAsset& GetSceneAsset() const { return *SceneAsset; }
+
+	const FContextualAnimIKTargetDefContainer& GetIKTargetDefs() const;
+
+	/** Return the current playback time of the animation this actor is playing */
+	float GetAnimMontageTime() const;
+
+	FName GetCurrentSection() const;
+
+	int32 GetCurrentSectionIndex() const;
+
+	/** Returns the ActiveMontageInstance or null in the case of static actors */
+	FAnimMontageInstance* GetAnimMontageInstance() const;
+
+	UAnimInstance* GetAnimInstance() const;
+
+	USkeletalMeshComponent* GetSkeletalMeshComponent() const;
+
+	UContextualAnimSceneActorComponent* GetSceneActorComponent() const;
+
+	static const FContextualAnimSceneBinding InvalidBinding;
+
+private:
+
+	friend UContextualAnimSceneInstance;
+
+	FContextualAnimSceneBindingContext Context;
+
+	/** Scene Asset we are bound to */
+	TWeakObjectPtr<const UContextualAnimSceneAsset> SceneAsset = nullptr;
+
+	/** Ptr to the animation data in the scene asset used by this actor */
+	const FContextualAnimTrack* AnimTrackPtr = nullptr;
+
+	/** Ptr to the role definition used by this actor */
+	const FContextualAnimRoleDefinition* RoleDefPtr = nullptr;
+
+	/** Ptr back to the scene instance we belong to (if any) */
+	TWeakObjectPtr<const UContextualAnimSceneInstance> SceneInstancePtr = nullptr;
+
+#if WITH_EDITOR
+public:
+	/** Guid only used in editor to bind this actor to sequencer */
+	FGuid Guid;
+#endif
+
+};
+
+USTRUCT(BlueprintType)
+struct CONTEXTUALANIMATION_API FContextualAnimSceneBindings
+{
+	GENERATED_BODY()
+
+	const FContextualAnimSceneBinding* FindBindingByActor(const AActor* Actor) const
+	{
+		return Actor ? Data.FindByPredicate([Actor](const auto& Item) { return Item.GetActor() == Actor; }) : nullptr;
+	}
+
+	const FContextualAnimSceneBinding* FindBindingByRole(const FName& Role) const
+	{
+		return Role != NAME_None ? Data.FindByPredicate([&Role](const auto& Item) { return Item.GetAnimTrack().Role == Role; }) : nullptr;
+	}
+
+#if WITH_EDITOR
+	const FContextualAnimSceneBinding* FindBindingByGuid(const FGuid& Guid) const
+	{
+		return Guid.IsValid() ? Data.FindByPredicate([&Guid](const auto& Item) { return Item.Guid == Guid; }) : nullptr;
+	}
+#endif
+
+	FORCEINLINE const UContextualAnimSceneAsset* GetSceneAsset() const { return Num() > 0 ? &Data[0].GetSceneAsset() : nullptr; }
+	FORCEINLINE int32 GetVariantIdx() const { return Num() > 0 ? Data[0].GetAnimTrack().VariantIdx : INDEX_NONE; }
+	FORCEINLINE int32 Num() const { return Data.Num(); }
+	FORCEINLINE int32 Add(const FContextualAnimSceneBinding& NewData) { return Data.Add(NewData); }
+	FORCEINLINE void Reset() { return Data.Reset(); }
+	FORCEINLINE const TArray<FContextualAnimSceneBinding>& GetBindings() const { return Data; }
+
+	FORCEINLINE TArray<FContextualAnimSceneBinding>::RangedForIteratorType      begin() { return Data.begin(); }
+	FORCEINLINE TArray<FContextualAnimSceneBinding>::RangedForConstIteratorType begin() const { return Data.begin(); }
+	FORCEINLINE TArray<FContextualAnimSceneBinding>::RangedForIteratorType      end() { return Data.end(); }
+	FORCEINLINE TArray<FContextualAnimSceneBinding>::RangedForConstIteratorType end() const { return Data.end(); }
+
+	static bool TryCreateBindings(const UContextualAnimSceneAsset& SceneAsset, int32 VariantIdx, const TMap<FName, FContextualAnimSceneBindingContext>& Params, FContextualAnimSceneBindings& OutBindings);
+
+private:
+
+	/** List of actors bound to each role in the SceneAsset */
+	UPROPERTY()
+	TArray<FContextualAnimSceneBinding> Data;
+};
+
 USTRUCT(BlueprintType)
 struct FContextualAnimStartSceneParams
 {
@@ -195,21 +394,16 @@ struct FContextualAnimStartSceneParams
 
 	/** Map with actors to bind to each role in the scene */
 	UPROPERTY(EditAnywhere, BlueprintReadWrite, Category = "Defaults")
-	TMap<FName, AActor*> RoleToActorMap;
+	TMap<FName, FContextualAnimSceneBindingContext> RoleToActorMap;
 
 	/** Desired variant. If INDEX_NONE the Manager will attempt to find the best variant to use by running the selection criteria */
 	UPROPERTY(EditAnywhere, BlueprintReadWrite, Category = "Defaults")
 	int32 VariantIdx = INDEX_NONE;
 
-	/** Desired start time */
-	UPROPERTY(EditAnywhere, BlueprintReadWrite, Category = "Defaults")
-	float AnimStartTime = 0.f;
-
 	void Reset()
 	{
 		RoleToActorMap.Reset();
 		VariantIdx = INDEX_NONE;
-		AnimStartTime = 0.f;
 	}
 };
 
@@ -272,28 +466,4 @@ struct FContextualAnimQueryParams
 
 	FContextualAnimQueryParams(const FTransform& InQueryTransform, bool bInComplexQuery, bool bInFindAnimStartTime)
 		: QueryTransform(InQueryTransform), bComplexQuery(bInComplexQuery), bFindAnimStartTime(bInFindAnimStartTime) {}
-};
-
-USTRUCT(BlueprintType)
-struct FContextualAnimQuerierData
-{
-	GENERATED_BODY()
-
-	UPROPERTY(EditAnywhere, BlueprintReadWrite, Category = "Defaults")
-	FTransform Transform;
-
-	UPROPERTY(EditAnywhere, BlueprintReadWrite, Category = "Defaults")
-	FVector Velocity = FVector::ZeroVector;
-
-	UPROPERTY(EditAnywhere, BlueprintReadWrite, Category = "Defaults")
-	FGameplayTagContainer GameplayTags;
-};
-
-USTRUCT(BlueprintType)
-struct FContextualAnimPrimaryActorData
-{
-	GENERATED_BODY()
-
-	UPROPERTY(EditAnywhere, BlueprintReadWrite, Category = "Defaults")
-	FTransform Transform;
 };
\ No newline at end of file
diff --git a/Engine/Plugins/Experimental/Animation/ContextualAnimation/Source/ContextualAnimation/Public/ContextualAnimUtilities.h b/Engine/Plugins/Experimental/Animation/ContextualAnimation/Source/ContextualAnimation/Public/ContextualAnimUtilities.h
--- a/Engine/Plugins/Experimental/Animation/ContextualAnimation/Source/ContextualAnimation/Public/ContextualAnimUtilities.h
+++ b/Engine/Plugins/Experimental/Animation/ContextualAnimation/Source/ContextualAnimation/Public/ContextualAnimUtilities.h
@@ -51,6 +51,16 @@ public:
 
 	static FAnimMontageInstance* TryGetActiveAnimMontageInstance(const AActor* Actor);
 
+	static void DrawSector(FPrimitiveDrawInterface& PDI, const FVector& Origin, const FVector& Direction, float MinDistance, float MaxDistance, float MinAngle, float MaxAngle, const FLinearColor& Color, uint8 DepthPriority, float Thickness);
+
+	static bool CalculateScenePivotForAlignmentSection(const FContextualAnimAlignmentSectionData& AligmentSectionData, const FContextualAnimSceneBindings& Bindings, FTransform& OutScenePivot);
+
+	UFUNCTION(BlueprintCallable, Category = "Contextual Anim|Scene Asset", meta = (DisplayName = "Create Contextual Anim Scene Bindings"))
+	static bool BP_CreateContextualAnimSceneBindings(const UContextualAnimSceneAsset* SceneAsset, const TMap<FName, FContextualAnimSceneBindingContext>& Params, FContextualAnimSceneBindings& OutBindings);
+
+	// Montage Blueprint Interface
+	//------------------------------------------------------------------------------------------
+
 	UFUNCTION(BlueprintCallable, Category = "Contextual Anim|Utilities", meta = (DisplayName = "GetSectionStartAndEndTime"))
 	static void BP_Montage_GetSectionStartAndEndTime(const UAnimMontage* Montage, int32 SectionIndex, float& OutStartTime, float& OutEndTime);
 
@@ -60,5 +70,46 @@ public:
 	UFUNCTION(BlueprintCallable, Category = "Contextual Anim|Utilities", meta = (DisplayName = "GetSectionLength"))
 	static float BP_Montage_GetSectionLength(const UAnimMontage* Montage, int32 SectionIndex);
 
-	static void DrawSector(FPrimitiveDrawInterface& PDI, const FVector& Origin, const FVector& Direction, float MinDistance, float MaxDistance, float MinAngle, float MaxAngle, const FLinearColor& Color, uint8 DepthPriority, float Thickness);
+	// SceneBindings Blueprint Interface
+	//------------------------------------------------------------------------------------------
+
+	UFUNCTION(BlueprintCallable, Category = "Contextual Anim|Scene Bindings", meta = (DisplayName = "Add Or Update Warp Targets For Bindings"))
+	static void BP_SceneBindings_AddOrUpdateWarpTargetsForBindings(const FContextualAnimSceneBindings& Bindings);
+
+	UFUNCTION(BlueprintPure, Category = "Contextual Anim|Scene Bindings", meta = (DisplayName = "Get Bindings"))
+	static const TArray<FContextualAnimSceneBinding>& BP_SceneBindings_GetBindings(const FContextualAnimSceneBindings& Bindings) { return Bindings.GetBindings(); }
+
+	UFUNCTION(BlueprintPure, Category = "Contextual Anim|Scene Bindings", meta = (DisplayName = "Get Binding By Role"))
+	static const FContextualAnimSceneBinding& BP_SceneBindings_GetBindingByRole(const FContextualAnimSceneBindings& Bindings, FName Role);
+
+	UFUNCTION(BlueprintPure, Category = "Contextual Anim|Scene Bindings", meta = (DisplayName = "Get Scene Asset"))
+	static const UContextualAnimSceneAsset* BP_SceneBindings_GetSceneAsset(const FContextualAnimSceneBindings& Bindings) { return Bindings.GetSceneAsset(); }
+
+	UFUNCTION(BlueprintPure, Category = "Contextual Anim|Scene Bindings", meta = (DisplayName = "Get Alignment Transform For Role Relative To Other Role"))
+	static FTransform BP_SceneBindings_GetAlignmentTransformForRoleRelativeToOtherRole(const FContextualAnimSceneBindings& Bindings, FName Role, FName RelativeToRole, float Time);
+
+	// FContextualAnimSceneBindingContext Blueprint Interface
+	//------------------------------------------------------------------------------------------
+	
+	UFUNCTION(BlueprintPure, Category = "Contextual Anim|Scene Binding Context", meta = (NativeMakeFunc, DisplayName = "Make Contextual Anim Scene Binding Context"))
+	static FContextualAnimSceneBindingContext BP_SceneBindingContext_MakeFromActor(AActor* Actor) { return FContextualAnimSceneBindingContext(Actor); }
+
+	UFUNCTION(BlueprintPure, Category = "Contextual Anim|Scene Binding Context", meta = (DisplayName = "Make Contextual Anim Scene Binding Context With External Transform"))
+	static FContextualAnimSceneBindingContext BP_SceneBindingContext_MakeFromActorWithExternalTransform(AActor* Actor, FTransform ExternalTransform) { return FContextualAnimSceneBindingContext(Actor, ExternalTransform); }
+
+	// FContextualAnimSceneBinding Blueprint Interface
+	//------------------------------------------------------------------------------------------
+
+	UFUNCTION(BlueprintPure, Category = "Contextual Anim|Scene Binding", meta = (DisplayName = "Get Role"))
+	static FName BP_SceneBinding_GetRole(const FContextualAnimSceneBinding& Binding) { return Binding.GetRoleDef().Name; }
+
+	UFUNCTION(BlueprintPure, Category = "Contextual Anim|Scene Binding", meta = (DisplayName = "Get Actor"))
+	static AActor* BP_SceneBinding_GetActor(const FContextualAnimSceneBinding& Binding) { return Binding.GetActor(); }
+
+	UFUNCTION(BlueprintPure, Category = "Contextual Anim|Scene Binding", meta = (DisplayName = "Get Skeletal Mesh"))
+	static USkeletalMeshComponent* BP_SceneBinding_GetSkeletalMesh(const FContextualAnimSceneBinding& Binding) { return Binding.GetSkeletalMeshComponent(); }
+
+	UFUNCTION(BlueprintPure, Category = "Contextual Anim|Scene Binding", meta = (DisplayName = "Get Animation To Play"))
+	static const UAnimSequenceBase* BP_SceneBinding_GetAnimationToPlay(const FContextualAnimSceneBinding& Binding) { return Binding.GetAnimTrack().Animation; }
+
 };
\ No newline at end of file
diff --git a/Engine/Plugins/Experimental/Animation/ContextualAnimation/Source/ContextualAnimationEditor/Private/ContextualAnimEdMode.cpp b/Engine/Plugins/Experimental/Animation/ContextualAnimation/Source/ContextualAnimationEditor/Private/ContextualAnimEdMode.cpp
--- a/Engine/Plugins/Experimental/Animation/ContextualAnimation/Source/ContextualAnimationEditor/Private/ContextualAnimEdMode.cpp
+++ b/Engine/Plugins/Experimental/Animation/ContextualAnimation/Source/ContextualAnimationEditor/Private/ContextualAnimEdMode.cpp
@@ -99,42 +99,32 @@ void FContextualAnimEdMode::Render(const FSceneView* View, FViewport* Viewport,
 			}
 
 			FTransform PrimaryActorTransform = FTransform::Identity;
-			if (const FContextualAnimSceneActorData* PreviewData = SceneInstance->FindSceneActorDataByRole(SceneAsset->GetPrimaryRole()))
+			if (const FContextualAnimSceneBinding* Binding = SceneInstance->FindBindingByRole(SceneAsset->GetPrimaryRole()))
 			{
-				PrimaryActorTransform = PreviewData->GetActor()->GetTransform();
+				PrimaryActorTransform = Binding->GetTransform();
 			}
 
 			const EShowIKTargetsDrawMode IKTargetsDrawMode = ViewportClient->GetShowIKTargetsDrawMode();
 
-			for (const FContextualAnimSceneActorData& SceneActorData : SceneInstance->GetBindings())
+			for (const FContextualAnimSceneBinding& Binding : SceneInstance->GetBindings())
 			{
 				// Draw IK Targets
-				if (IKTargetsDrawMode == EShowIKTargetsDrawMode::All || (IKTargetsDrawMode == EShowIKTargetsDrawMode::Selected && SceneActorData.GetActor() == SelectedActor.Get()))
+				if (IKTargetsDrawMode == EShowIKTargetsDrawMode::All || (IKTargetsDrawMode == EShowIKTargetsDrawMode::Selected && Binding.GetActor() == SelectedActor.Get()))
 				{
-					DrawIKTargetsForSceneActor(*PDI, SceneActorData);
+					DrawIKTargetsForBinding(*PDI, Binding);
 				}
 
 				// Draw Selection Criteria
-				if (const FContextualAnimTrack* AnimTrack = SceneAsset->GetAnimTrack(SceneActorData.GetRole(), VariantIdx))
+				if (const FContextualAnimTrack* AnimTrack = SceneAsset->GetAnimTrack(Binding.GetRoleDef().Name, VariantIdx))
 				{
 					for (int32 CriterionIdx = 0; CriterionIdx < AnimTrack->SelectionCriteria.Num(); CriterionIdx++)
 					{
 						if(const UContextualAnimSelectionCriterion* Criterion = AnimTrack->SelectionCriteria[CriterionIdx])
 						{
 							FLinearColor DrawColor = FLinearColor::White;
-
+							if (Criterion->DoesQuerierPassCondition(FContextualAnimSceneBindingContext(PrimaryActorTransform), Binding.GetContext()))
 							{
-								FContextualAnimQuerierData QuerierData;
-								QuerierData.Transform = SceneActorData.GetActor()->GetActorTransform();
-								QuerierData.Velocity = SceneActorData.GetActor()->GetVelocity();
-
-								FContextualAnimPrimaryActorData PrimaryActorData;
-								PrimaryActorData.Transform = PrimaryActorTransform;
-
-								if (Criterion->DoesQuerierPassCondition(PrimaryActorData, QuerierData))
-								{
-									DrawColor = FLinearColor::Green;
-								}
+								DrawColor = FLinearColor::Green;
 							}
 
 							//@TODO: Each SelectionCriterion should implement this, and here we should just call "Criterion->Draw()"
@@ -153,11 +143,11 @@ void FContextualAnimEdMode::Render(const FSceneView* View, FViewport* Viewport,
 
 									PDI->DrawLine(P0, P0 + FVector::UpVector * Spatial->Height, DrawColor, SDPG_Foreground, 2.f);
 
-									PDI->SetHitProxy(new HSelectionCriterionHitProxy(FSelectionCriterionHitProxyData(SceneActorData.GetRole(), VariantIdx, CriterionIdx, Idx)));
+									PDI->SetHitProxy(new HSelectionCriterionHitProxy(FSelectionCriterionHitProxyData(Binding.GetRoleDef().Name, VariantIdx, CriterionIdx, Idx)));
 									PDI->DrawPoint(P0, FLinearColor::Black, 15.f, SDPG_Foreground);
 									PDI->SetHitProxy(nullptr);
 
-									PDI->SetHitProxy(new HSelectionCriterionHitProxy(FSelectionCriterionHitProxyData(SceneActorData.GetRole(), VariantIdx, CriterionIdx, Idx + 4)));
+									PDI->SetHitProxy(new HSelectionCriterionHitProxy(FSelectionCriterionHitProxyData(Binding.GetRoleDef().Name, VariantIdx, CriterionIdx, Idx + 4)));
 									PDI->DrawPoint(P0 + FVector::UpVector * Spatial->Height, FLinearColor::Black, 15.f, SDPG_Foreground);
 									PDI->SetHitProxy(nullptr);
 								}
@@ -175,15 +165,15 @@ void FContextualAnimEdMode::Render(const FSceneView* View, FViewport* Viewport,
 	}
 }
 
-void FContextualAnimEdMode::DrawIKTargetsForSceneActor(FPrimitiveDrawInterface& PDI, const FContextualAnimSceneActorData& SceneActorData) const
+void FContextualAnimEdMode::DrawIKTargetsForBinding(FPrimitiveDrawInterface& PDI, const FContextualAnimSceneBinding& Binding) const
 {
-	for (const FContextualAnimIKTargetDefinition& IKTargetDef : SceneActorData.GetIKTargetDefs().IKTargetDefs)
+	for (const FContextualAnimIKTargetDefinition& IKTargetDef : Binding.GetIKTargetDefs().IKTargetDefs)
 	{
-		if (const FContextualAnimSceneActorData* TargetSceneActorData = SceneActorData.GetSceneInstance().FindSceneActorDataByRole(IKTargetDef.TargetRoleName))
+		if (const FContextualAnimSceneBinding* TargetBinding = Binding.GetSceneInstance()->FindBindingByRole(IKTargetDef.TargetRoleName))
 		{
-			if (USkeletalMeshComponent* TargetSkelMeshComp = TargetSceneActorData->GetSkeletalMeshComponent())
+			if (USkeletalMeshComponent* TargetSkelMeshComp = TargetBinding->GetSkeletalMeshComponent())
 			{
-				const float Alpha = UAnimNotifyState_IKWindow::GetIKAlphaValue(IKTargetDef.GoalName, SceneActorData.GetAnimMontageInstance());
+				const float Alpha = UAnimNotifyState_IKWindow::GetIKAlphaValue(IKTargetDef.GoalName, Binding.GetAnimMontageInstance());
 
 				if (IKTargetDef.Provider == EContextualAnimIKTargetProvider::Bone)
 				{
@@ -193,7 +183,7 @@ void FContextualAnimEdMode::DrawIKTargetsForSceneActor(FPrimitiveDrawInterface&
 				{
 					const FTransform ParentTransform = TargetSkelMeshComp->GetSocketTransform(IKTargetDef.TargetBoneName);
 
-					const FTransform TargetTransform = SceneActorData.GetAnimTrack().IKTargetData.ExtractTransformAtTime(IKTargetDef.GoalName, SceneActorData.GetAnimTime()) * ParentTransform;
+					const FTransform TargetTransform = Binding.GetAnimTrack().IKTargetData.ExtractTransformAtTime(IKTargetDef.GoalName, Binding.GetAnimMontageTime()) * ParentTransform;
 
 					FLinearColor Color = Alpha > 0.f ? FLinearColor(FColor::MakeRedToGreenColorFromScalar(Alpha)) : FLinearColor::White;
 
@@ -355,9 +345,9 @@ bool FContextualAnimEdMode::GetCustomDrawingCoordinateSystem(FMatrix& InMatrix,
 			if (const UContextualAnimSceneInstance* SceneInstance = ViewModel->GetSceneInstance())
 			{
 				FTransform PrimaryActorTransform = FTransform::Identity;
-				if (const FContextualAnimSceneActorData* PreviewData = SceneInstance->FindSceneActorDataByRole(ViewModel->GetSceneAsset()->GetPrimaryRole()))
+				if (const FContextualAnimSceneBinding* Binding = SceneInstance->FindBindingByRole(ViewModel->GetSceneAsset()->GetPrimaryRole()))
 				{
-					PrimaryActorTransform = PreviewData->GetActor()->GetTransform();
+					PrimaryActorTransform = Binding->GetTransform();
 					InMatrix = PrimaryActorTransform.ToMatrixNoScale().RemoveTranslation();
 					return true;
 				}
@@ -399,9 +389,9 @@ FVector FContextualAnimEdMode::GetWidgetLocation() const
 							}
 
 							FTransform PrimaryActorTransform = FTransform::Identity;
-							if (const FContextualAnimSceneActorData* PreviewData = SceneInstance->FindSceneActorDataByRole(SceneAsset->GetPrimaryRole()))
+							if (const FContextualAnimSceneBinding* Binding = SceneInstance->FindBindingByRole(SceneAsset->GetPrimaryRole()))
 							{
-								PrimaryActorTransform = PreviewData->GetActor()->GetTransform();
+								PrimaryActorTransform = Binding->GetTransform();
 							}
 
 							return PrimaryActorTransform.TransformPositionNoScale(Location);
diff --git a/Engine/Plugins/Experimental/Animation/ContextualAnimation/Source/ContextualAnimationEditor/Private/ContextualAnimViewModel.cpp b/Engine/Plugins/Experimental/Animation/ContextualAnimation/Source/ContextualAnimationEditor/Private/ContextualAnimViewModel.cpp
--- a/Engine/Plugins/Experimental/Animation/ContextualAnimation/Source/ContextualAnimationEditor/Private/ContextualAnimViewModel.cpp
+++ b/Engine/Plugins/Experimental/Animation/ContextualAnimation/Source/ContextualAnimationEditor/Private/ContextualAnimViewModel.cpp
@@ -76,7 +76,7 @@ void FContextualAnimViewModel::Initialize(UContextualAnimSceneAsset* InSceneAsse
 
 UAnimSequenceBase* FContextualAnimViewModel::FindAnimationByGuid(const FGuid& Guid) const
 {
-	return SceneInstance.IsValid() ? SceneInstance->FindSceneActorDataByGuid(Guid)->GetAnimTrack().Animation : nullptr;
+	return SceneInstance.IsValid() ? SceneInstance->FindBindingByGuid(Guid)->GetAnimTrack().Animation : nullptr;
 }
 
 void FContextualAnimViewModel::CreateSequencer()
@@ -184,9 +184,10 @@ void FContextualAnimViewModel::RefreshSequencerTracks()
 	{
 		for (auto& MapEntry : StartSceneParams.RoleToActorMap)
 		{
-			if (MapEntry.Value)
+			
+			if (AActor* Actor = MapEntry.Value.GetActor())
 			{
-				MapEntry.Value->Destroy();
+				Actor->Destroy();
 			}
 		}
 	}
@@ -362,16 +363,16 @@ UObject* FContextualAnimViewModel::GetPlaybackContext() const
 
 void FContextualAnimViewModel::SequencerTimeChanged()
 {
-	auto ResetActorTransform = [](FContextualAnimSceneActorData& SceneActorData, float Time)
+	auto ResetActorTransform = [](FContextualAnimSceneBinding& Binding, float Time)
 	{
-		const USkeletalMeshComponent* SkelMeshComp = UContextualAnimUtilities::TryGetSkeletalMeshComponent(SceneActorData.GetActor());
+		const USkeletalMeshComponent* SkelMeshComp = UContextualAnimUtilities::TryGetSkeletalMeshComponent(Binding.GetActor());
 
-		const FTransform RootTransform = UContextualAnimUtilities::ExtractRootTransformFromAnimation(SceneActorData.GetAnimTrack().Animation, Time);
+		const FTransform RootTransform = UContextualAnimUtilities::ExtractRootTransformFromAnimation(Binding.GetAnimTrack().Animation, Time);
 		const FTransform StartTransform = SkelMeshComp->GetRelativeTransform().Inverse() * RootTransform;
 
-		SceneActorData.GetActor()->SetActorLocationAndRotation(StartTransform.GetLocation(), StartTransform.GetRotation());
+		Binding.GetActor()->SetActorLocationAndRotation(StartTransform.GetLocation(), StartTransform.GetRotation());
 
-		if (UCharacterMovementComponent* MovementComp = SceneActorData.GetActor()->FindComponentByClass<UCharacterMovementComponent>())
+		if (UCharacterMovementComponent* MovementComp = Binding.GetActor()->FindComponentByClass<UCharacterMovementComponent>())
 		{
 			MovementComp->StopMovementImmediately();
 		}
@@ -492,7 +493,7 @@ void FContextualAnimViewModel::OnPreviewActorClassChanged()
 	{
 		for (const auto& Binding : SceneInstance->GetBindings())
 		{
-			if (const FContextualAnimRoleDefinition* RoleDef = RolesAsset->FindRoleDefinitionByName(Binding.GetRole()))
+			if (const FContextualAnimRoleDefinition* RoleDef = RolesAsset->FindRoleDefinitionByName(Binding.GetRoleDef().Name))
 			{
 				const UClass* DesiredPreviewClass = RoleDef->PreviewActorClass;
 				const UClass* CurrentPreviewClass = Binding.GetActor()->GetClass();
diff --git a/Engine/Plugins/Experimental/Animation/ContextualAnimation/Source/ContextualAnimationEditor/Public/ContextualAnimEdMode.h b/Engine/Plugins/Experimental/Animation/ContextualAnimation/Source/ContextualAnimationEditor/Public/ContextualAnimEdMode.h
--- a/Engine/Plugins/Experimental/Animation/ContextualAnimation/Source/ContextualAnimationEditor/Public/ContextualAnimEdMode.h
+++ b/Engine/Plugins/Experimental/Animation/ContextualAnimation/Source/ContextualAnimationEditor/Public/ContextualAnimEdMode.h
@@ -77,7 +77,7 @@ public:
 
 	bool GetHitResultUnderCursor(FHitResult& OutHitResult, FEditorViewportClient* InViewportClient, const FViewportClick& Click) const;
 
-	void DrawIKTargetsForSceneActor(class FPrimitiveDrawInterface& PDI, const struct FContextualAnimSceneActorData& SceneActorData) const;
+	void DrawIKTargetsForBinding(class FPrimitiveDrawInterface& PDI, const struct FContextualAnimSceneBinding& Binding) const;
 
 private:
 
