Contextual Anim Editor (WIP): Updated new variant popup to use AnimSequenceBase instead of Montage.

diff --git a/Engine/Plugins/Experimental/Animation/ContextualAnimation/Source/ContextualAnimationEditor/Public/ContextualAnimEditorTypes.h b/Engine/Plugins/Experimental/Animation/ContextualAnimation/Source/ContextualAnimationEditor/Public/ContextualAnimEditorTypes.h
--- a/Engine/Plugins/Experimental/Animation/ContextualAnimation/Source/ContextualAnimationEditor/Public/ContextualAnimEditorTypes.h
+++ b/Engine/Plugins/Experimental/Animation/ContextualAnimation/Source/ContextualAnimationEditor/Public/ContextualAnimEditorTypes.h
@@ -67,7 +67,7 @@ struct FContextualAnimNewVariantData
 	FName RoleName = NAME_None;
 
 	UPROPERTY(EditAnywhere, Category = "Settings")
-	UAnimMontage* Animation = nullptr;
+	UAnimSequenceBase* Animation = nullptr;
 
 	UPROPERTY(EditAnywhere, Category = "Settings")
 	bool bRequiresFlyingMode = false;
