Contextual Anim (WIP): Added support for multiple sections and renamed Variants to AnimSet and AlignmentSection to AnimSetPivotDefinition

diff --git a/Engine/Plugins/Experimental/Animation/ContextualAnimation/Source/ContextualAnimation/Private/ContextualAnimManager.cpp b/Engine/Plugins/Experimental/Animation/ContextualAnimation/Source/ContextualAnimation/Private/ContextualAnimManager.cpp
--- a/Engine/Plugins/Experimental/Animation/ContextualAnimation/Source/ContextualAnimation/Private/ContextualAnimManager.cpp
+++ b/Engine/Plugins/Experimental/Animation/ContextualAnimation/Source/ContextualAnimation/Private/ContextualAnimManager.cpp
@@ -100,6 +100,12 @@
 
 UContextualAnimSceneInstance* UContextualAnimManager::ForceStartScene(const UContextualAnimSceneAsset& SceneAsset, const FContextualAnimStartSceneParams& Params)
 {
+	if(Params.RoleToActorMap.Num() == 0)
+	{
+		UE_LOG(LogContextualAnim, Warning, TEXT("UContextualAnimManager::ForceStartScene. Can't start scene. Reason: Empty Params. SceneAsset: %s"), *GetNameSafe(&SceneAsset));
+		return nullptr;
+	}
+
 	FContextualAnimSceneBindings Bindings;
 	for (const auto& Pair : Params.RoleToActorMap)
 	{
@@ -114,7 +120,8 @@
 			return nullptr;
 		}
 
-		const FContextualAnimTrack* AnimTrack = SceneAsset.GetAnimTrack(RoleToBind, Params.VariantIdx);
+		const int32 SectionIdx = 0; // Always start from the first section
+		const FContextualAnimTrack* AnimTrack = SceneAsset.GetAnimTrack(SectionIdx, Params.AnimSetIdx, RoleToBind);
 		if (AnimTrack == nullptr)
 		{
 			UE_LOG(LogContextualAnim, Warning, TEXT("UContextualAnimManager::ForceStartScene. Can't start scene. Reason: Can't find anim track for '%s'. SceneAsset: %s"), 
@@ -151,19 +158,20 @@
 
 UContextualAnimSceneInstance* UContextualAnimManager::TryStartScene(const UContextualAnimSceneAsset& SceneAsset, const FContextualAnimStartSceneParams& Params)
 {
+	const int32 SectionIdx = 0; // Always start from the first section
 	FContextualAnimSceneBindings Bindings;
 
 	bool bSuccess = false;
-	if (Params.VariantIdx != INDEX_NONE)
+	if (Params.AnimSetIdx != INDEX_NONE)
 	{
-		bSuccess = FContextualAnimSceneBindings::TryCreateBindings(SceneAsset, Params.VariantIdx, Params.RoleToActorMap, Bindings);
+		bSuccess = FContextualAnimSceneBindings::TryCreateBindings(SceneAsset, SectionIdx, Params.AnimSetIdx, Params.RoleToActorMap, Bindings);
 	}
 	else
 	{
-		const int32 VariantsNum = SceneAsset.GetTotalVariants();
-		for (int32 VariantIdx = 0; VariantIdx < VariantsNum; VariantIdx++)
+		const int32 NumSets = SceneAsset.GetNumAnimSetsInSection(SectionIdx);
+		for (int32 AnimSetIdx = 0; AnimSetIdx < NumSets; AnimSetIdx++)
 		{
-			if (FContextualAnimSceneBindings::TryCreateBindings(SceneAsset, VariantIdx, Params.RoleToActorMap, Bindings))
+			if (FContextualAnimSceneBindings::TryCreateBindings(SceneAsset, SectionIdx, AnimSetIdx, Params.RoleToActorMap, Bindings))
 			{
 				bSuccess = true;
 				break;
diff --git a/Engine/Plugins/Experimental/Animation/ContextualAnimation/Source/ContextualAnimation/Private/ContextualAnimSceneActorComponent.cpp b/Engine/Plugins/Experimental/Animation/ContextualAnimation/Source/ContextualAnimation/Private/ContextualAnimSceneActorComponent.cpp
--- a/Engine/Plugins/Experimental/Animation/ContextualAnimation/Source/ContextualAnimation/Private/ContextualAnimSceneActorComponent.cpp
+++ b/Engine/Plugins/Experimental/Animation/ContextualAnimation/Source/ContextualAnimation/Private/ContextualAnimSceneActorComponent.cpp
@@ -28,7 +28,9 @@
 
 FBoxSphereBounds UContextualAnimSceneActorComponent::CalcBounds(const FTransform& LocalToWorld) const
 {
-	const float Radius = SceneAsset ? SceneAsset->GetRadius() : 0.f;
+	// Use radius of the first set in the first section for now. 
+	// The option of having an SceneAsset and draw options on this component may go away in the future anyway, replaced by smart objects.
+	const float Radius = SceneAsset && SceneAsset->HasValidData() ? SceneAsset->GetSection(0)->GetAnimSet(0)->Radius : 0.f;
 	return FBoxSphereBounds(FSphere(GetComponentTransform().GetLocation(), Radius));
 }
 
diff --git a/Engine/Plugins/Experimental/Animation/ContextualAnimation/Source/ContextualAnimation/Private/ContextualAnimSceneAsset.cpp b/Engine/Plugins/Experimental/Animation/ContextualAnimation/Source/ContextualAnimation/Private/ContextualAnimSceneAsset.cpp
--- a/Engine/Plugins/Experimental/Animation/ContextualAnimation/Source/ContextualAnimation/Private/ContextualAnimSceneAsset.cpp
+++ b/Engine/Plugins/Experimental/Animation/ContextualAnimation/Source/ContextualAnimation/Private/ContextualAnimSceneAsset.cpp
@@ -25,118 +25,72 @@
 	return FCompactPoseBoneIndex(INDEX_NONE);
 }
 
-UContextualAnimSceneAsset::UContextualAnimSceneAsset(const FObjectInitializer& ObjectInitializer)
-	: Super(ObjectInitializer)
+// FContextualAnimSceneSection
+//==============================================================================================
+
+const FContextualAnimSet* FContextualAnimSceneSection::GetAnimSet(int32 AnimSetIdx) const
 {
-	bDisableCollisionBetweenActors = true;
-	SampleRate = 15;
-
-	AlignmentSections.Add(FContextualAnimAlignmentSectionData());
+	return AnimSets.IsValidIndex(AnimSetIdx) ? &AnimSets[AnimSetIdx] : nullptr;
 }
 
-#if WITH_EDITOR
-void UContextualAnimSceneAsset::PostEditChangeProperty(FPropertyChangedEvent& PropertyChangedEvent)
+void FContextualAnimSceneSection::GenerateAlignmentTracks(UContextualAnimSceneAsset& SceneAsset)
 {
-	Super::PostEditChangeProperty(PropertyChangedEvent);
-
-	const FName PropertyName = (PropertyChangedEvent.Property != nullptr) ? PropertyChangedEvent.Property->GetFName() : NAME_None;
-	const FName MemberPropertyName = (PropertyChangedEvent.MemberProperty != nullptr) ? PropertyChangedEvent.MemberProperty->GetFName() : NAME_None;
-
-	UE_LOG(LogContextualAnim, Log, TEXT("%s::PostEditChangeProperty MemberPropertyName: %s PropertyName: %s"),
-		*GetNameSafe(this), *MemberPropertyName.ToString(), *PropertyName.ToString());
-
-	if(MemberPropertyName == GET_MEMBER_NAME_CHECKED(UContextualAnimSceneAsset, AlignmentSections))
-	{
-		GenerateAlignmentTracks();
-	}
-}
-#endif
-
-void UContextualAnimSceneAsset::PreSave(FObjectPreSaveContext ObjectSaveContext)
-{
-	Super::PreSave(ObjectSaveContext);
-
-	PrecomputeData();
-}
-
-void UContextualAnimSceneAsset::PrecomputeData()
-{
-	// Set VariantIdx for each AnimTrack
-	for (int32 VariantIdx = 0; VariantIdx < Variants.Num(); VariantIdx++)
-	{
-		for (FContextualAnimTrack& AnimTrack : Variants[VariantIdx].Tracks)
-		{
-			AnimTrack.VariantIdx = VariantIdx;
-		}
-	}
-
-	// Generate alignment tracks relative to scene pivot
-	GenerateAlignmentTracks();
-
-	// Generate IK Targets
-	GenerateIKTargetTracks();
-
-	UpdateRadius();
-}
-
-void UContextualAnimSceneAsset::GenerateAlignmentTracks()
-{
 	// Necessary for FCompactPose that uses a FAnimStackAllocator (TMemStackAllocator) which allocates from FMemStack.
 	// When allocating memory from FMemStack we need to explicitly use FMemMark to ensure items are freed when the scope exits. 
 	// UWorld::Tick adds a FMemMark to catch any allocation inside the game tick 
 	// but any allocation from outside the game tick (like here when generating the alignment tracks off-line) must explicitly add a mark to avoid a leak 
 	FMemMark Mark(FMemStack::Get());
 
-	for (FContextualAnimTracksContainer& Variant : Variants)
+	for (FContextualAnimSet& AnimSet : AnimSets)
 	{
-		Variant.ScenePivots.Reset();
+		AnimSet.ScenePivots.Reset();
 
-		// Generate scene pivot for each alignment section
-		for (const FContextualAnimAlignmentSectionData& AligmentSection : AlignmentSections)
+		// Generate pivot for each alignment section
+		for (const FContextualAnimSetPivotDefinition& Def : AnimSetPivotDefinitions)
 		{
 			FTransform ScenePivot = FTransform::Identity;
-			FContextualAnimTrack* AnimTrack = Variant.Tracks.FindByPredicate([&AligmentSection](const FContextualAnimTrack& AnimTrack) { return AnimTrack.Role == AligmentSection.Origin; });
+			FContextualAnimTrack* AnimTrack = AnimSet.Tracks.FindByPredicate([&Def](const FContextualAnimTrack& AnimTrack) { return AnimTrack.Role == Def.Origin; });
 			if (AnimTrack)
 			{
-				if (AligmentSection.bAlongClosestDistance)
+				if (Def.bAlongClosestDistance)
 				{
-					FContextualAnimTrack* OtherAnimTrack = Variant.Tracks.FindByPredicate([&AligmentSection](const FContextualAnimTrack& AnimTrack) { return AnimTrack.Role == AligmentSection.OtherRole; });
+					FContextualAnimTrack* OtherAnimTrack = AnimSet.Tracks.FindByPredicate([&Def](const FContextualAnimTrack& AnimTrack) { return AnimTrack.Role == Def.OtherRole; });
 					if (OtherAnimTrack)
 					{
 						FTransform T1 = AnimTrack->Animation ? UContextualAnimUtilities::ExtractRootTransformFromAnimation(AnimTrack->Animation, 0.f) : FTransform::Identity;
-						T1 = (GetMeshToComponentForRole(AnimTrack->Role).Inverse() * (T1 * AnimTrack->MeshToScene));
+						T1 = (SceneAsset.GetMeshToComponentForRole(AnimTrack->Role).Inverse() * (T1 * AnimTrack->MeshToScene));
 
 						FTransform T2 = OtherAnimTrack->Animation ? UContextualAnimUtilities::ExtractRootTransformFromAnimation(OtherAnimTrack->Animation, 0.f) : FTransform::Identity;
-						T2 = (GetMeshToComponentForRole(OtherAnimTrack->Role).Inverse() * (T2 * OtherAnimTrack->MeshToScene));
+						T2 = (SceneAsset.GetMeshToComponentForRole(OtherAnimTrack->Role).Inverse() * (T2 * OtherAnimTrack->MeshToScene));
 
-						ScenePivot.SetLocation(FMath::Lerp<FVector>(T1.GetLocation(), T2.GetLocation(), AligmentSection.Weight));
+						ScenePivot.SetLocation(FMath::Lerp<FVector>(T1.GetLocation(), T2.GetLocation(), Def.Weight));
 						ScenePivot.SetRotation((T2.GetLocation() - T1.GetLocation()).GetSafeNormal2D().ToOrientationQuat());
 					}
 				}
 				else
 				{
 					const FTransform RootTransform = AnimTrack->Animation ? UContextualAnimUtilities::ExtractRootTransformFromAnimation(AnimTrack->Animation, 0.f) : FTransform::Identity;
-					ScenePivot = (GetMeshToComponentForRole(AnimTrack->Role).Inverse() * (RootTransform * AnimTrack->MeshToScene));
+					ScenePivot = (SceneAsset.GetMeshToComponentForRole(AnimTrack->Role).Inverse() * (RootTransform * AnimTrack->MeshToScene));
 				}
 			}
 
-			Variant.ScenePivots.Add(ScenePivot);
+			AnimSet.ScenePivots.Add(ScenePivot);
 		}
 
 		// Generate alignment tracks relative to scene pivot
-		for (FContextualAnimTrack& AnimTrack : Variant.Tracks)
+		for (FContextualAnimTrack& AnimTrack : AnimSet.Tracks)
 		{
-			UE_LOG(LogContextualAnim, Log, TEXT("%s Generating AlignmentTracks Tracks. Animation: %s"), *GetNameSafe(this), *GetNameSafe(AnimTrack.Animation));
+			UE_LOG(LogContextualAnim, Log, TEXT("Generating AlignmentTracks Tracks. Animation: %s"), *GetNameSafe(AnimTrack.Animation));
 
-			const FTransform MeshToComponentInverse = GetMeshToComponentForRole(AnimTrack.Role).Inverse();
-			const float SampleInterval = 1.f / SampleRate;
+			const FTransform MeshToComponentInverse = SceneAsset.GetMeshToComponentForRole(AnimTrack.Role).Inverse();
+			const float SampleInterval = 1.f / SceneAsset.GetSampleRate();
 
 			// Initialize tracks for each alignment section
-			const int32 TotalTracks = AlignmentSections.Num();
+			const int32 TotalTracks = AnimSetPivotDefinitions.Num();
 			AnimTrack.AlignmentData.Initialize(TotalTracks, SampleInterval);
 			for (int32 Idx = 0; Idx < TotalTracks; Idx++)
 			{
-				AnimTrack.AlignmentData.Tracks.TrackNames.Add(AlignmentSections[Idx].WarpTargetName);
+				AnimTrack.AlignmentData.Tracks.TrackNames.Add(AnimSetPivotDefinitions[Idx].Name);
 				AnimTrack.AlignmentData.Tracks.AnimationTracks.AddZeroed();
 			}
 
@@ -156,7 +110,7 @@
 					{
 						FRawAnimSequenceTrack& AlignmentTrack = AnimTrack.AlignmentData.Tracks.AnimationTracks[Idx];
 
-						const FTransform ScenePivotTransform = Variants[AnimTrack.VariantIdx].ScenePivots[Idx];
+						const FTransform ScenePivotTransform = AnimSets[AnimTrack.AnimSetIdx].ScenePivots[Idx];
 						const FTransform RootRelativeToScenePivot = RootTransform.GetRelativeTransform(ScenePivotTransform);
 
 						AlignmentTrack.PosKeys.Add(FVector3f(RootRelativeToScenePivot.GetLocation()));
@@ -172,7 +126,7 @@
 				{
 					FRawAnimSequenceTrack& SceneTrack = AnimTrack.AlignmentData.Tracks.AnimationTracks[Idx];
 
-					const FTransform ScenePivotTransform = Variants[AnimTrack.VariantIdx].ScenePivots[Idx];
+					const FTransform ScenePivotTransform = AnimSets[AnimTrack.AnimSetIdx].ScenePivots[Idx];
 					const FTransform RootRelativeToScenePivot = RootTransform.GetRelativeTransform(ScenePivotTransform);
 
 					SceneTrack.PosKeys.Add(FVector3f(RootRelativeToScenePivot.GetLocation()));
@@ -183,13 +137,17 @@
 	}
 }
 
-void UContextualAnimSceneAsset::GenerateIKTargetTracks()
+void FContextualAnimSceneSection::GenerateIKTargetTracks(UContextualAnimSceneAsset& SceneAsset)
 {
+	// @TODO: Optimize
+	// We are generating IK Target tracks for the duration of the animation and for all the sets in the section
+	// We should generate those tracks only for the animations with IK windows and only for the range of that window
+
 	FMemMark Mark(FMemStack::Get());
 
-	for (FContextualAnimTracksContainer& Variant : Variants)
+	for (FContextualAnimSet& AnimSet : AnimSets)
 	{
-		for (FContextualAnimTrack& AnimTrack : Variant.Tracks)
+		for (FContextualAnimTrack& AnimTrack : AnimSet.Tracks)
 		{
 			AnimTrack.IKTargetData.Empty();
 
@@ -201,9 +159,9 @@
 
 			if (const UAnimSequenceBase* Animation = AnimTrack.Animation)
 			{
-				UE_LOG(LogContextualAnim, Log, TEXT("%s Generating IK Target Tracks. Animation: %s"), *GetNameSafe(this), *GetNameSafe(Animation));
+				UE_LOG(LogContextualAnim, Log, TEXT("Generating IK Target Tracks. Animation: %s"), *GetNameSafe(Animation));
 
-				const float SampleInterval = 1.f / SampleRate;
+				const float SampleInterval = 1.f / SceneAsset.GetSampleRate();
 
 				TArray<FBoneIndexType> RequiredBoneIndexArray;
 
@@ -226,7 +184,7 @@
 						if (DataPtr == nullptr)
 						{
 							// Find track for target role. 
-							const FContextualAnimTrack* TargetAnimTrackPtr = GetAnimTrack(TargetRole, AnimTrack.VariantIdx);
+							const FContextualAnimTrack* TargetAnimTrackPtr = GetAnimTrack(AnimTrack.AnimSetIdx, TargetRole);
 							if (TargetAnimTrackPtr == nullptr)
 							{
 								UE_LOG(LogContextualAnim, Warning, TEXT("\t Can't find AnimTrack for TargetRole '%s'"), *TargetRole.ToString());
@@ -363,36 +321,59 @@
 	}
 }
 
-void UContextualAnimSceneAsset::UpdateRadius()
+const FContextualAnimTrack* FContextualAnimSceneSection::GetAnimTrack(int32 AnimSetIdx, const FName& Role) const
 {
-	Radius = 0.f;
-	ForEachAnimTrack([this](const FContextualAnimTrack& AnimTrack)
-		{
-			Radius = FMath::Max(Radius, (AnimTrack.GetAlignmentTransformAtEntryTime().GetLocation().Size()));
-			return UE::ContextualAnim::EForEachResult::Continue;
-		});
+	if (AnimSets.IsValidIndex(AnimSetIdx))
+	{
+		return AnimSets[AnimSetIdx].Tracks.FindByPredicate([Role](const FContextualAnimTrack& AnimTrack) { return AnimTrack.Role == Role; });
+	}
+
+	return nullptr;
 }
 
-FName UContextualAnimSceneAsset::FindRoleByAnimation(const UAnimSequenceBase* Animation) const
+const FContextualAnimTrack* FContextualAnimSceneSection::GetAnimTrack(int32 AnimSetIdx, int32 AnimTrackIdx) const
 {
-	FName Result = NAME_None;
+	if (AnimSets.IsValidIndex(AnimSetIdx))
+	{
+		const FContextualAnimSet& AnimSet = AnimSets[AnimSetIdx];
+		if (AnimSet.Tracks.IsValidIndex(AnimTrackIdx))
+		{
+			return &AnimSet.Tracks[AnimTrackIdx];
+		}
+	}
 
-	ForEachAnimTrack([&Result, Animation](const FContextualAnimTrack& AnimTrack)
-		{
-			if(AnimTrack.Animation == Animation)
-			{
-				Result = AnimTrack.Role;
-				return UE::ContextualAnim::EForEachResult::Break;
-			}
-
-			return UE::ContextualAnim::EForEachResult::Continue;
-		});
-
-	return Result;
+	return nullptr;
 }
 
-const FContextualAnimIKTargetDefContainer& UContextualAnimSceneAsset::GetIKTargetDefsForRole(const FName& Role) const
+FTransform FContextualAnimSceneSection::GetAlignmentTransformForRoleRelativeToPivot(int32 AnimSetIdx, FName Role, float Time) const
 {
+	if (const FContextualAnimTrack* AnimTrack = GetAnimTrack(AnimSetIdx, Role))
+	{
+		return AnimTrack->GetAlignmentTransformAtTime(Time);
+	}
+
+	return FTransform::Identity;
+}
+
+FTransform FContextualAnimSceneSection::GetAlignmentTransformForRoleRelativeToOtherRole(int32 AnimSetIdx, FName Role, FName OtherRole, float Time) const
+{
+	const FTransform FromRoleRelativeToScenePivot = GetAlignmentTransformForRoleRelativeToPivot(AnimSetIdx, Role, Time);
+	const FTransform ToRoleRelativeToScenePivot = GetAlignmentTransformForRoleRelativeToPivot(AnimSetIdx, OtherRole, Time);
+	return FromRoleRelativeToScenePivot.GetRelativeTransform(ToRoleRelativeToScenePivot);
+}
+
+FTransform FContextualAnimSceneSection::GetIKTargetTransformForRoleAtTime(int32 AnimSetIdx, FName Role, FName TrackName, float Time) const
+{
+	if (const FContextualAnimTrack* AnimTrack = GetAnimTrack(AnimSetIdx, Role))
+	{
+		return AnimTrack->IKTargetData.ExtractTransformAtTime(TrackName, Time);
+	}
+
+	return FTransform::Identity;
+}
+
+const FContextualAnimIKTargetDefContainer& FContextualAnimSceneSection::GetIKTargetDefsForRole(const FName& Role) const
+{
 	if (const FContextualAnimIKTargetDefContainer* ContainerPtr = RoleToIKTargetDefsMap.Find(Role))
 	{
 		return *ContainerPtr;
@@ -401,78 +382,11 @@
 	return FContextualAnimIKTargetDefContainer::EmptyContainer;
 }
 
-const FTransform& UContextualAnimSceneAsset::GetMeshToComponentForRole(const FName& Role) const
+const FContextualAnimTrack* FContextualAnimSceneSection::FindFirstAnimTrackForRoleThatPassesSelectionCriteria(const FName& Role, const FContextualAnimSceneBindingContext& Primary, const FContextualAnimSceneBindingContext& Querier) const
 {
-	if (RolesAsset)
+	for (const FContextualAnimSet& AnimSet : AnimSets)
 	{
-		if (const FContextualAnimRoleDefinition* RoleDef = RolesAsset->FindRoleDefinitionByName(Role))
-		{
-			return RoleDef->MeshToComponent;
-		}
-	}
-
-	return FTransform::Identity;
-}
-
-const FContextualAnimTrack* UContextualAnimSceneAsset::GetAnimTrack(const FName& Role, int32 VariantIdx) const
-{
-	if(Variants.IsValidIndex(VariantIdx))
-	{
-		return Variants[VariantIdx].Tracks.FindByPredicate([Role](const FContextualAnimTrack& AnimTrack){ return AnimTrack.Role == Role;});
-	}
-
-	return nullptr;
-}
-
-void UContextualAnimSceneAsset::ForEachAnimTrack(FForEachAnimTrackFunction Function) const
-{
-	for (const FContextualAnimTracksContainer& TracksContainer : Variants)
-	{
-		for (const FContextualAnimTrack& AnimTrack : TracksContainer.Tracks)
-		{
-			if (Function(AnimTrack) == UE::ContextualAnim::EForEachResult::Break)
-			{
-				return;
-			}
-		}
-	}
-}
-
-void UContextualAnimSceneAsset::ForEachAnimTrack(int32 VariantIdx, FForEachAnimTrackFunction Function) const
-{
-	if(Variants.IsValidIndex(VariantIdx))
-	{
-		const TArray<FContextualAnimTrack>& Tracks = Variants[VariantIdx].Tracks;
-		for (const FContextualAnimTrack& AnimTrack : Tracks)
-		{
-			if (Function(AnimTrack) == UE::ContextualAnim::EForEachResult::Break)
-			{
-				return;
-			}
-		}
-	}
-}
-
-TArray<FName> UContextualAnimSceneAsset::GetRoles() const
-{
- 	TArray<FName> Result;
-
-	if(RolesAsset)
-	{
-		for (const FContextualAnimRoleDefinition& RoleDef : RolesAsset->Roles)
-		{
-			Result.Add(RoleDef.Name);
-		}
-	}
-
- 	return Result;
-}
-
-const FContextualAnimTrack* UContextualAnimSceneAsset::FindFirstAnimTrackForRoleThatPassesSelectionCriteria(const FName& Role, const FContextualAnimSceneBindingContext& Primary, const FContextualAnimSceneBindingContext& Querier) const
-{
-	for (const FContextualAnimTracksContainer& Variant : Variants)
-	{
-		for (const FContextualAnimTrack& AnimTrack : Variant.Tracks)
+		for (const FContextualAnimTrack& AnimTrack : AnimSet.Tracks)
 		{
 			if (AnimTrack.Role == Role)
 			{
@@ -497,14 +411,14 @@
 	return nullptr;
 }
 
-const FContextualAnimTrack* UContextualAnimSceneAsset::FindAnimTrackForRoleWithClosestEntryLocation(const FName& Role, const FContextualAnimSceneBindingContext& Primary, const FVector& TestLocation) const
+const FContextualAnimTrack* FContextualAnimSceneSection::FindAnimTrackForRoleWithClosestEntryLocation(const FName& Role, const FContextualAnimSceneBindingContext& Primary, const FVector& TestLocation) const
 {
 	const FContextualAnimTrack* Result = nullptr;
 
 	float BestDistanceSq = MAX_FLT;
-	for (const FContextualAnimTracksContainer& Variant : Variants)
+	for (const FContextualAnimSet& AnimSet : AnimSets)
 	{
-		for (const FContextualAnimTrack& AnimTrack : Variant.Tracks)
+		for (const FContextualAnimTrack& AnimTrack : AnimSet.Tracks)
 		{
 			if (AnimTrack.Role == Role)
 			{
@@ -522,6 +436,148 @@
 	return Result;
 }
 
+// UContextualAnimSceneAsset
+//==============================================================================================
+
+UContextualAnimSceneAsset::UContextualAnimSceneAsset(const FObjectInitializer& ObjectInitializer)
+	: Super(ObjectInitializer)
+{
+	bDisableCollisionBetweenActors = true;
+	SampleRate = 15;
+}
+
+void UContextualAnimSceneAsset::PreSave(FObjectPreSaveContext ObjectSaveContext)
+{
+	Super::PreSave(ObjectSaveContext);
+
+	PrecomputeData();
+}
+
+void UContextualAnimSceneAsset::PrecomputeData()
+{
+	// Set SetIdx and SectionIdx on each AnimTrack and update the Radius for each set
+	for (int32 SectionIdx = 0; SectionIdx < Sections.Num(); SectionIdx++)
+	{
+		FContextualAnimSceneSection& Section = Sections[SectionIdx];
+		for (int32 SetIdx = 0; SetIdx < Section.AnimSets.Num(); SetIdx++)
+		{
+			FContextualAnimSet& AnimSet = Section.AnimSets[SetIdx];
+
+			FBox Box(ForceInitToZero);
+			for (FContextualAnimTrack& AnimTrack : AnimSet.Tracks)
+			{
+				AnimTrack.SectionIdx = SectionIdx;
+				AnimTrack.AnimSetIdx = SetIdx;
+
+				Box += AnimTrack.GetRootTransformAtTime(0.f).GetLocation();
+			}
+
+			AnimSet.Radius = Box.GetExtent().Size();
+		}
+
+		// Generate alignment tracks
+		Section.GenerateAlignmentTracks(*this);
+
+		// Generate IK Target tracks
+		Section.GenerateIKTargetTracks(*this);
+	}
+}
+
+const FContextualAnimTrack* UContextualAnimSceneAsset::FindAnimTrackByAnimation(const UAnimSequenceBase* Animation) const
+{
+	const FContextualAnimTrack* Result = nullptr;
+
+	ForEachAnimTrack([&Result, Animation](const FContextualAnimTrack& AnimTrack)
+		{
+			if(AnimTrack.Animation == Animation)
+			{
+				Result = &AnimTrack;
+				return UE::ContextualAnim::EForEachResult::Break;
+			}
+
+			return UE::ContextualAnim::EForEachResult::Continue;
+		});
+
+	return Result;
+}
+
+const FContextualAnimSceneSection* UContextualAnimSceneAsset::GetSection(int32 SectionIdx) const
+{
+	return (Sections.IsValidIndex(SectionIdx)) ? &Sections[SectionIdx] : nullptr;
+}
+
+const FContextualAnimSceneSection* UContextualAnimSceneAsset::GetSection(const FName& SectionName) const
+{
+	return Sections.FindByPredicate([&SectionName](const FContextualAnimSceneSection& Section) { return Section.Name == SectionName; });
+}
+
+const FContextualAnimTrack* UContextualAnimSceneAsset::GetAnimTrack(int32 SectionIdx, int32 AnimSetIdx, const FName& Role) const
+{
+	return (Sections.IsValidIndex(SectionIdx)) ? Sections[SectionIdx].GetAnimTrack(AnimSetIdx, Role) : nullptr;
+}
+
+const FContextualAnimTrack* UContextualAnimSceneAsset::GetAnimTrack(int32 SectionIdx, int32 AnimSetIdx, int32 AnimTrackIdx) const
+{
+	return (Sections.IsValidIndex(SectionIdx)) ? Sections[SectionIdx].GetAnimTrack(AnimSetIdx, AnimTrackIdx) : nullptr;
+}
+
+void UContextualAnimSceneAsset::ForEachAnimTrack(FForEachAnimTrackFunction Function) const
+{
+	for (const FContextualAnimSceneSection& Section : Sections)
+	{
+		for (const FContextualAnimSet& Set : Section.AnimSets)
+		{
+			for (const FContextualAnimTrack& AnimTrack : Set.Tracks)
+			{
+				if (Function(AnimTrack) == UE::ContextualAnim::EForEachResult::Break)
+				{
+					return;
+				}
+			}
+		}
+	}
+}
+
+TArray<FName> UContextualAnimSceneAsset::GetRoles() const
+{
+ 	TArray<FName> Result;
+
+	if(RolesAsset)
+	{
+		for (const FContextualAnimRoleDefinition& RoleDef : RolesAsset->Roles)
+		{
+			Result.Add(RoleDef.Name);
+		}
+	}
+
+ 	return Result;
+}
+
+const FTransform& UContextualAnimSceneAsset::GetMeshToComponentForRole(const FName& Role) const
+{
+	if (RolesAsset)
+	{
+		if (const FContextualAnimRoleDefinition* RoleDef = RolesAsset->FindRoleDefinitionByName(Role))
+		{
+			return RoleDef->MeshToComponent;
+		}
+	}
+
+	return FTransform::Identity;
+}
+
+TArray<FName> UContextualAnimSceneAsset::GetSectionNames() const
+{
+	TArray<FName> Result;
+
+	for (const FContextualAnimSceneSection& Section : Sections)
+	{
+		Result.Add(Section.Name);
+	}
+
+	return Result;
+}
+
 bool UContextualAnimSceneAsset::Query(FName Role, FContextualAnimQueryResult& OutResult, const FContextualAnimQueryParams& QueryParams, const FTransform& ToWorldTransform) const
 {
 	//@TODO: Kept around only to do not break existing content. It will go away in the future.
@@ -529,9 +585,10 @@
 	FContextualAnimSceneBindingContext Primary(ToWorldTransform);
 	FContextualAnimSceneBindingContext Querier(QueryParams.Querier.IsValid() ? QueryParams.Querier->GetActorTransform() : QueryParams.QueryTransform);
 
-	if (const FContextualAnimTrack* AnimTrack = FindFirstAnimTrackForRoleThatPassesSelectionCriteria(Role, Primary, Querier))
+	const FContextualAnimTrack* AnimTrack = Sections.IsValidIndex(0) ? Sections[0].FindFirstAnimTrackForRoleThatPassesSelectionCriteria(Role, Primary, Querier) : nullptr;
+	if (AnimTrack)
 	{
-		OutResult.VariantIdx = AnimTrack->VariantIdx;
+		OutResult.AnimSetIdx = AnimTrack->AnimSetIdx;
 		OutResult.Animation = Cast<UAnimMontage>(AnimTrack->Animation);
 		OutResult.EntryTransform = AnimTrack->GetAlignmentTransformAtEntryTime() * ToWorldTransform;
 		OutResult.SyncTransform = AnimTrack->GetAlignmentTransformAtSyncTime() * ToWorldTransform;
@@ -548,52 +605,66 @@
 	return false;
 }
 
-FTransform UContextualAnimSceneAsset::GetAlignmentTransformForRoleRelativeToScenePivot(FName Role, int32 VariantIdx, float Time) const
+int32 UContextualAnimSceneAsset::GetNumSections() const
 {
-	if (const FContextualAnimTrack* AnimTrack = GetAnimTrack(Role, VariantIdx))
+	return Sections.Num();
+}
+
+int32 UContextualAnimSceneAsset::GetNumAnimSetsInSection(int32 SectionIdx) const
+{
+	return Sections.IsValidIndex(SectionIdx) ? Sections[SectionIdx].AnimSets.Num() : 0;
+}
+
+const FContextualAnimIKTargetDefContainer& UContextualAnimSceneAsset::GetIKTargetDefsForRoleInSection(int32 SectionIdx, const FName& Role) const
+{
+	return Sections.IsValidIndex(SectionIdx) ? Sections[SectionIdx].GetIKTargetDefsForRole(Role) : FContextualAnimIKTargetDefContainer::EmptyContainer;
+}
+
+FTransform UContextualAnimSceneAsset::GetAlignmentTransformForRoleRelativeToOtherRoleInSection(int32 SectionIdx, int32 AnimSetIdx, const FName& Role, const FName& OtherRole, float Time) const
+{
+	return Sections.IsValidIndex(SectionIdx) ? Sections[SectionIdx].GetAlignmentTransformForRoleRelativeToOtherRole(AnimSetIdx, Role, OtherRole, Time) : FTransform::Identity;
+}
+
+const TArray<FContextualAnimSetPivotDefinition>& UContextualAnimSceneAsset::GetAnimSetPivotDefinitionsInSection(int32 SectionIdx) const
+{
+	static TArray<FContextualAnimSetPivotDefinition> EmptyDefs;
+	return Sections.IsValidIndex(SectionIdx) ? Sections[SectionIdx].GetAnimSetPivotDefinitions() : EmptyDefs;
+}
+
+// Blueprint Interface
+//------------------------------------------------------------------------------------------
+
+UAnimSequenceBase* UContextualAnimSceneAsset::BP_FindAnimationForRole(int32 SectionIdx, int32 AnimSetIdx, FName Role) const
+{
+	const FContextualAnimTrack* AnimTrack = GetAnimTrack(SectionIdx, AnimSetIdx, Role);
+	return AnimTrack ? AnimTrack->Animation : nullptr;
+}
+
+FTransform UContextualAnimSceneAsset::BP_GetAlignmentTransformForRoleRelativeToPivot(int32 SectionIdx, int32 AnimSetIdx, FName Role, float Time) const
+{
+	return (Sections.IsValidIndex(SectionIdx)) ? Sections[SectionIdx].GetAlignmentTransformForRoleRelativeToPivot(AnimSetIdx, Role, Time) : FTransform::Identity;
+}
+
+FTransform UContextualAnimSceneAsset::BP_GetIKTargetTransformForRoleAtTime(int32 SectionIdx, int32 AnimSetIdx, FName Role, FName TrackName, float Time) const
+{
+	return (Sections.IsValidIndex(SectionIdx)) ? Sections[SectionIdx].GetIKTargetTransformForRoleAtTime(AnimSetIdx, Role, TrackName, Time) : FTransform::Identity;
+}
+
+int32 UContextualAnimSceneAsset::BP_FindAnimSetIndexByAnimation(int32 SectionIdx, const UAnimSequenceBase* Animation) const
+{
+	if (Sections.IsValidIndex(SectionIdx))
 	{
-		return AnimTrack->GetAlignmentTransformAtTime(Time);
+		for (const FContextualAnimSet& Set : Sections[SectionIdx].AnimSets)
+		{
+			for (const FContextualAnimTrack& AnimTrack : Set.Tracks)
+			{
+				if (AnimTrack.Animation == Animation)
+				{
+					return AnimTrack.AnimSetIdx;
+				}
+			}
+		}
 	}
 
-	return FTransform::Identity;
+	return INDEX_NONE;
 }
-
-FTransform UContextualAnimSceneAsset::GetAlignmentTransformForRoleRelativeToOtherRole(FName FromRole, FName ToRole, int32 VariantIdx, float Time) const
-{
-	const FTransform FromRoleRelativeToScenePivot = GetAlignmentTransformForRoleRelativeToScenePivot(FromRole, VariantIdx, Time);
-	const FTransform ToRoleRelativeToScenePivot = GetAlignmentTransformForRoleRelativeToScenePivot(ToRole, VariantIdx, Time);
-	return FromRoleRelativeToScenePivot.GetRelativeTransform(ToRoleRelativeToScenePivot);
-}
-
-FTransform UContextualAnimSceneAsset::GetIKTargetTransformForRoleAtTime(FName Role, int32 VariantIdx, FName TrackName, float Time) const
-{
-	if(const FContextualAnimTrack* AnimTrack = GetAnimTrack(Role, VariantIdx))
-	{
-		return AnimTrack->IKTargetData.ExtractTransformAtTime(TrackName, Time);		
-	}
-
-	return FTransform::Identity;
-}
-
-int32 UContextualAnimSceneAsset::FindVariantIdx(FName Role, UAnimSequenceBase* Animation) const
-{
-	int32 Result = INDEX_NONE;
-
-	ForEachAnimTrack([&Result, Role, Animation](const FContextualAnimTrack& AnimTrack)
-		{
-			if (AnimTrack.Role == Role && AnimTrack.Animation == Animation)
-			{
-				Result = AnimTrack.VariantIdx;
-				return UE::ContextualAnim::EForEachResult::Break;
-			}
-
-			return UE::ContextualAnim::EForEachResult::Continue;
-		});
-
-	return Result;
-}
-
-int32 UContextualAnimSceneAsset::GetTotalVariants() const
-{
-	return Variants.Num();
-}
diff --git a/Engine/Plugins/Experimental/Animation/ContextualAnimation/Source/ContextualAnimation/Private/ContextualAnimSceneInstance.cpp b/Engine/Plugins/Experimental/Animation/ContextualAnimation/Source/ContextualAnimation/Private/ContextualAnimSceneInstance.cpp
--- a/Engine/Plugins/Experimental/Animation/ContextualAnimation/Source/ContextualAnimation/Private/ContextualAnimSceneInstance.cpp
+++ b/Engine/Plugins/Experimental/Animation/ContextualAnimation/Source/ContextualAnimation/Private/ContextualAnimSceneInstance.cpp
@@ -43,6 +43,27 @@
 	return Binding ? Binding->GetActor() : nullptr;
 }
 
+UAnimMontage* UContextualAnimSceneInstance::PlayAnimation(UAnimInstance& AnimInstance, UAnimSequenceBase& Animation)
+{
+	if (UAnimMontage* AnimMontage = Cast<UAnimMontage>(&Animation))
+	{
+		const float Duration = AnimInstance.Montage_Play(AnimMontage, 1.f, EMontagePlayReturnType::MontageLength);
+		return (Duration > 0.f) ? AnimMontage : nullptr;
+	}
+	else
+	{
+		//@TODO: Expose all these on the AnimTrack
+		const FName SlotName = FName(TEXT("DefaultSlot"));
+		const float BlendInTime = 0.25f;
+		const float BlendOutTime = 0.25f;
+		const float InPlayRate = 1.f;
+		const int32 LoopCount = 1;
+		const float BlendOutTriggerTime = -1.f;
+		const float InTimeToStartMontageAt = 0.f;
+		return AnimInstance.PlaySlotAnimationAsDynamicMontage(&Animation, SlotName, BlendInTime, BlendOutTime, InPlayRate, LoopCount, BlendOutTriggerTime, InTimeToStartMontageAt);
+	}
+}
+
 void UContextualAnimSceneInstance::Join(FContextualAnimSceneBinding& Binding)
 {
 	AActor* Actor = Binding.GetActor();
@@ -55,30 +76,14 @@
 	{
 		if (UAnimInstance* AnimInstance = Binding.GetAnimInstance())
 		{
-			// Keep montage support for now but might go away soon
-			if (UAnimMontage* AnimMontage = Cast<UAnimMontage>(Animation))
+			if(UAnimMontage* Montage = PlayAnimation(*AnimInstance, *Animation))
 			{
-				AnimInstance->Montage_Play(AnimMontage, 1.f, EMontagePlayReturnType::MontageLength);
+				AnimInstance->OnPlayMontageNotifyBegin.AddUniqueDynamic(this, &UContextualAnimSceneInstance::OnNotifyBeginReceived);
+				AnimInstance->OnPlayMontageNotifyEnd.AddUniqueDynamic(this, &UContextualAnimSceneInstance::OnNotifyEndReceived);
+				AnimInstance->OnMontageBlendingOut.AddUniqueDynamic(this, &UContextualAnimSceneInstance::OnMontageBlendingOut);
 			}
-			else
-			{
-				//@TODO: Expose all these on the AnimTrack
-				const FName SlotName = FName(TEXT("DefaultSlot"));
-				const float BlendInTime = 0.25f;
-				const float BlendOutTime = 0.25f;
-				const float InPlayRate = 1.f;
-				const int32 LoopCount = 1;
-				const float BlendOutTriggerTime = -1.f;
-				const float InTimeToStartMontageAt = 0.f;
-				AnimInstance->PlaySlotAnimationAsDynamicMontage(Binding.GetAnimTrack().Animation, SlotName, BlendInTime, BlendOutTime, InPlayRate, LoopCount, BlendOutTriggerTime, InTimeToStartMontageAt);
-			}
-
-			AnimInstance->OnPlayMontageNotifyBegin.AddUniqueDynamic(this, &UContextualAnimSceneInstance::OnNotifyBeginReceived);
-			AnimInstance->OnPlayMontageNotifyEnd.AddUniqueDynamic(this, &UContextualAnimSceneInstance::OnNotifyEndReceived);
-			AnimInstance->OnMontageBlendingOut.AddUniqueDynamic(this, &UContextualAnimSceneInstance::OnMontageBlendingOut);
 		}
 
-
 		//@TODO: Temp, until we have a way to switch between movement mode using AnimNotifyState
 		if (Binding.GetAnimTrack().bRequireFlyingMode)
 		{
@@ -97,8 +102,7 @@
 			const FTransform ScenePivotRuntime = Pair.Value;
 			const float Time = Binding.GetAnimTrack().GetSyncTimeForWarpSection(WarpTargetName);
 			const FTransform TransformRelativeToScenePivot = Binding.GetAnimTrack().AlignmentData.ExtractTransformAtTime(WarpTargetName, Time);
-			const FTransform WarpTarget = (TransformRelativeToScenePivot * ScenePivotRuntime);
-
+			const FTransform WarpTarget = TransformRelativeToScenePivot * ScenePivotRuntime;
 			MotionWarpComp->AddOrUpdateWarpTargetFromTransform(WarpTargetName, WarpTarget);
 		}
 	}
@@ -120,56 +124,32 @@
 
 void UContextualAnimSceneInstance::Leave(FContextualAnimSceneBinding& Binding)
 {
-	// Check if we have an exit section and transition to it, otherwise just stop the montage
-	static const FName ExitSectionName = FName(TEXT("Exit"));
-	if(TransitionTo(Binding, ExitSectionName) == false)
+	if (UAnimInstance* AnimInstance = Binding.GetAnimInstance())
 	{
-		AActor* Actor = Binding.GetActor();
-		if (UAnimInstance* AnimInstance = Binding.GetAnimInstance())
+		if (UAnimMontage* CurrentMontage = AnimInstance->GetCurrentActiveMontage())
 		{
-			UAnimMontage* CurrentMontage = AnimInstance->GetCurrentActiveMontage();
-			if (CurrentMontage)
-			{
-				AnimInstance->Montage_Stop(CurrentMontage->BlendOut.GetBlendTime(), CurrentMontage);
-			}
+			AnimInstance->Montage_Stop(CurrentMontage->BlendOut.GetBlendTime(), CurrentMontage);
 		}
 	}
 }
 
-bool UContextualAnimSceneInstance::TransitionTo(FContextualAnimSceneBinding& Binding, const FName& ToSectionName)
+bool UContextualAnimSceneInstance::TransitionTo(FContextualAnimSceneBinding& Binding, const FContextualAnimTrack& AnimTrack)
 {
+	check(AnimTrack.Animation != Binding.GetAnimTrack().Animation);
+	check(AnimTrack.Role == Binding.GetRoleDef().Name);
+
 	UAnimInstance* AnimInstance = Binding.GetAnimInstance();
 	if (AnimInstance == nullptr)
 	{		
 		return false;
 	}
 
-	UAnimMontage* CurrentMontage = AnimInstance->GetCurrentActiveMontage();
-	if (CurrentMontage == nullptr)
-	{
-		UE_LOG(LogContextualAnim, Log, TEXT("UContextualAnimSceneInstance::TransitionTo. Actor is not playing any montage. Actor: %s ToSectionName: %s"),
-			*GetNameSafe(Binding.GetActor()), *ToSectionName.ToString());
-
-		return false;
-	}
-
-	const int32 SectionIdx = CurrentMontage->GetSectionIndex(ToSectionName);
-	if (SectionIdx == INDEX_NONE)
-	{
-		UE_LOG(LogContextualAnim, Log, TEXT("UContextualAnimSceneInstance::TransitionTo. Invalid Section. Actor: %s CurrentMontage: %s ToSectionName: %s"),
-			*GetNameSafe(Binding.GetActor()), *GetNameSafe(CurrentMontage), *ToSectionName.ToString());
-
-		return false;
-	}
-
-	UE_LOG(LogContextualAnim, Verbose, TEXT("UContextualAnimSceneInstance::TransitionTo. Actor: %s CurrentMontage: %s ToSectionName: %s"),
-		*GetNameSafe(Binding.GetActor()), *GetNameSafe(CurrentMontage), *ToSectionName.ToString());
-
 	// Unbind blend out delegate for a moment so we don't get it during the transition
+	// @TODO: Replace this with the TGuardValue 'pattern', similar to what we do in the editor for OnAnimNotifyChanged
 	AnimInstance->OnMontageBlendingOut.RemoveDynamic(this, &UContextualAnimSceneInstance::OnMontageBlendingOut);
 
-	AnimInstance->Montage_Play(CurrentMontage, 1.f);
-	AnimInstance->Montage_JumpToSection(ToSectionName, CurrentMontage);
+	PlayAnimation(*AnimInstance, *AnimTrack.Animation);
+	Binding.AnimTrackPtr = &AnimTrack;
 
 	AnimInstance->OnMontageBlendingOut.AddUniqueDynamic(this, &UContextualAnimSceneInstance::OnMontageBlendingOut);
 
@@ -178,31 +158,13 @@
 
 void UContextualAnimSceneInstance::Start()
 {
-	for (const FContextualAnimAlignmentSectionData AligmentData : SceneAsset->GetAlignmentSections())
+	const int32 SectionIdx = 0; // Always start from the first section
+	for (const FContextualAnimSetPivotDefinition& Def : SceneAsset->GetAnimSetPivotDefinitionsInSection(SectionIdx))
 	{
 		FTransform ScenePivotRuntime = FTransform::Identity;
-		if (const FContextualAnimSceneBinding* Binding = FindBindingByRole(AligmentData.Origin))
-		{
-			if (AligmentData.bAlongClosestDistance)
-			{
-				if (const FContextualAnimSceneBinding* OtherBinding = FindBindingByRole(AligmentData.OtherRole))
-				{
-					const FTransform T1 = Binding->GetTransform();
-					const FTransform T2 = OtherBinding->GetTransform();
-
-					ScenePivotRuntime.SetLocation(FMath::Lerp<FVector>(T1.GetLocation(), T2.GetLocation(), AligmentData.Weight));
-					ScenePivotRuntime.SetRotation((T2.GetLocation() - T1.GetLocation()).GetSafeNormal2D().ToOrientationQuat());
-				}
-			}
-			else
-			{
-				ScenePivotRuntime = Binding->GetTransform();
-			}
-		}
-
-		AlignmentSectionToScenePivotList.Add(MakeTuple(AligmentData.WarpTargetName, ScenePivotRuntime));
+		UContextualAnimUtilities::CalculateAnimSetPivot(Def, Bindings, ScenePivotRuntime);
+		AlignmentSectionToScenePivotList.Add(MakeTuple(Def.Name, ScenePivotRuntime));
 	}
-	
 
 	for (auto& Binding : Bindings)
 	{
diff --git a/Engine/Plugins/Experimental/Animation/ContextualAnimation/Source/ContextualAnimation/Private/ContextualAnimTypes.cpp b/Engine/Plugins/Experimental/Animation/ContextualAnimation/Source/ContextualAnimation/Private/ContextualAnimTypes.cpp
--- a/Engine/Plugins/Experimental/Animation/ContextualAnimation/Source/ContextualAnimation/Private/ContextualAnimTypes.cpp
+++ b/Engine/Plugins/Experimental/Animation/ContextualAnimation/Source/ContextualAnimation/Private/ContextualAnimTypes.cpp
@@ -11,6 +11,7 @@
 #include "ContextualAnimSelectionCriterion.h"
 #include "ContextualAnimSceneActorComponent.h"
 #include "ContextualAnimSceneAsset.h"
+#include "GameFramework/Character.h"
 
 DEFINE_LOG_CATEGORY(LogContextualAnim);
 
@@ -188,6 +189,17 @@
 	return true;
 }
 
+FTransform FContextualAnimTrack::GetRootTransformAtTime(float Time) const
+{
+	FTransform RootTransform = FTransform::Identity;
+	if (Animation)
+	{
+		RootTransform = UContextualAnimUtilities::ExtractRootTransformFromAnimation(Animation, Time);
+	}
+
+	return RootTransform * MeshToScene;
+}
+
 // FContextualAnimSceneBindingContext
 ///////////////////////////////////////////////////////////////////////
 
@@ -290,13 +302,13 @@
 
 const FContextualAnimIKTargetDefContainer& FContextualAnimSceneBinding::GetIKTargetDefs() const
 {
-	return GetSceneAsset().GetIKTargetDefsForRole(GetAnimTrack().Role);
+	return GetSceneAsset().GetIKTargetDefsForRoleInSection(GetAnimTrack().SectionIdx, GetAnimTrack().Role);
 }
 
 // FContextualAnimSceneBindings
 ///////////////////////////////////////////////////////////////////////
 
-bool FContextualAnimSceneBindings::TryCreateBindings(const UContextualAnimSceneAsset& SceneAsset, int32 VariantIdx, const TMap<FName, FContextualAnimSceneBindingContext>& Params, FContextualAnimSceneBindings& OutBindings)
+bool FContextualAnimSceneBindings::TryCreateBindings(const UContextualAnimSceneAsset& SceneAsset, int32 SectionIdx, int32 AnimSetIdx, const TMap<FName, FContextualAnimSceneBindingContext>& Params, FContextualAnimSceneBindings& OutBindings)
 {
 	check(SceneAsset.HasValidData());
 
@@ -316,15 +328,15 @@
 	// First, try to bind primary track. 
 	// @TODO: Revisit this, passing the same data twice (as primary and querier) feels weird, but this allow us to run the selection mechanism even on the primary actor.
 
-	const FContextualAnimTrack* PrimaryAnimTrack = SceneAsset.GetAnimTrack(PrimaryRole, VariantIdx);
+	const FContextualAnimTrack* PrimaryAnimTrack = SceneAsset.GetAnimTrack(SectionIdx, AnimSetIdx, PrimaryRole);
 	if (PrimaryAnimTrack && PrimaryAnimTrack->DoesQuerierPassSelectionCriteria(*PrimaryPtr, *PrimaryPtr))
 	{
 		OutBindings.Add(FContextualAnimSceneBinding(*PrimaryPtr, SceneAsset, *PrimaryAnimTrack));
 	}
 	else
 	{
-		UE_LOG(LogContextualAnim, Warning, TEXT("FContextualAnimSceneBindings::TryCreateBindings Failed. Reason: Can't find valid track for primary actor. SceneAsset: %s Role: %s Actor: %s Variant: %d"),
-			*GetNameSafe(&SceneAsset), *PrimaryRole.ToString(), *GetNameSafe(PrimaryPtr->GetActor()), VariantIdx);
+		UE_LOG(LogContextualAnim, Warning, TEXT("FContextualAnimSceneBindings::TryCreateBindings Failed. Reason: Can't find valid track for primary actor. SceneAsset: %s Role: %s Actor: %s SectionIdx: %d AnimSetIdx: %d"),
+			*GetNameSafe(&SceneAsset), *PrimaryRole.ToString(), *GetNameSafe(PrimaryPtr->GetActor()), SectionIdx, AnimSetIdx);
 
 		return false;
 	}
@@ -336,15 +348,15 @@
 		FName RoleToBind = Pair.Key;
 		if (RoleToBind != PrimaryRole)
 		{
-			const FContextualAnimTrack* AnimTrack = SceneAsset.GetAnimTrack(RoleToBind, VariantIdx);
+			const FContextualAnimTrack* AnimTrack = SceneAsset.GetAnimTrack(SectionIdx, AnimSetIdx, RoleToBind);
 			if (AnimTrack && AnimTrack->DoesQuerierPassSelectionCriteria(*PrimaryPtr, Pair.Value))
 			{
 				OutBindings.Add(FContextualAnimSceneBinding(Pair.Value, SceneAsset, *AnimTrack));
 			}
 			else
 			{
-				UE_LOG(LogContextualAnim, Warning, TEXT("FContextualAnimSceneBindings::TryCreateBindings Failed. Reason: Can't find valid track for secondary actor. SceneAsset: %s Role: %s Actor: %s Variant: %d"),
-					*GetNameSafe(&SceneAsset), *RoleToBind.ToString(), *GetNameSafe(Pair.Value.GetActor()), VariantIdx);
+				UE_LOG(LogContextualAnim, Warning, TEXT("FContextualAnimSceneBindings::TryCreateBindings Failed. Reason: Can't find valid track for secondary actor. SceneAsset: %s Role: %s Actor: %s SectionIdx: %d AnimSetIdx: %d"),
+					*GetNameSafe(&SceneAsset), *RoleToBind.ToString(), *GetNameSafe(Pair.Value.GetActor()), SectionIdx, AnimSetIdx);
 
 				return false;
 			}
diff --git a/Engine/Plugins/Experimental/Animation/ContextualAnimation/Source/ContextualAnimation/Private/ContextualAnimUtilities.cpp b/Engine/Plugins/Experimental/Animation/ContextualAnimation/Source/ContextualAnimation/Private/ContextualAnimUtilities.cpp
--- a/Engine/Plugins/Experimental/Animation/ContextualAnimation/Source/ContextualAnimation/Private/ContextualAnimUtilities.cpp
+++ b/Engine/Plugins/Experimental/Animation/ContextualAnimation/Source/ContextualAnimation/Private/ContextualAnimUtilities.cpp
@@ -131,14 +131,14 @@
 	}
 }
 
-void UContextualAnimUtilities::DrawDebugScene(const UWorld* World, const UContextualAnimSceneAsset* SceneAsset, int32 VariantIdx, float Time, const FTransform& ToWorldTransform, const FColor& Color, float LifeTime, float Thickness)
+void UContextualAnimUtilities::DrawDebugAnimSet(const UWorld* World, const UContextualAnimSceneAsset& SceneAsset, const FContextualAnimSet& AnimSet, float Time, const FTransform& ToWorldTransform, const FColor& Color, float LifeTime, float Thickness)
 {
-	if (World && SceneAsset)
+	if (World)
 	{
-		SceneAsset->ForEachAnimTrack(VariantIdx, [=](const FContextualAnimTrack& AnimTrack)
+		for(const FContextualAnimTrack& AnimTrack : AnimSet.Tracks)
 		{
-			const FTransform Transform = (SceneAsset->GetMeshToComponentForRole(AnimTrack.Role) * AnimTrack.GetAlignmentTransformAtTime(Time)) * ToWorldTransform;
-			
+			const FTransform Transform = (SceneAsset.GetMeshToComponentForRole(AnimTrack.Role) * AnimTrack.GetAlignmentTransformAtTime(Time)) * ToWorldTransform;
+
 			if (const UAnimSequenceBase* Animation = AnimTrack.Animation)
 			{
 				DrawDebugPose(World, Animation, Time, Transform, Color, LifeTime, Thickness);
@@ -147,9 +147,7 @@
 			{
 				DrawDebugCoordinateSystem(World, Transform.GetLocation(), Transform.Rotator(), 50.f, false, LifeTime, 0, Thickness);
 			}
-
-			return UE::ContextualAnim::EForEachResult::Continue;
-		});
+		}
 	}
 }
 
@@ -263,14 +261,16 @@
 {
 	if(SceneAsset == nullptr || !SceneAsset->HasValidData())
 	{
-		UE_LOG(LogContextualAnim, Warning, TEXT("UContextualAnimUtilities::FindBestVariantForActors Failed. Reason: Invalid or Empty SceneAsset (%s)"), *GetNameSafe(SceneAsset));
+		UE_LOG(LogContextualAnim, Warning, TEXT("UContextualAnimUtilities::BP_CreateContextualAnimSceneBindings Failed. Reason: Invalid or Empty SceneAsset (%s)"), *GetNameSafe(SceneAsset));
 		return false;
 	}
 
-	for (int32 VariantIdx = 0; VariantIdx < SceneAsset->GetTotalVariants(); VariantIdx++)
+	const int32 SectionIdx = 0; // Always start from the first section
+	const int32 NumSets = SceneAsset->GetNumAnimSetsInSection(SectionIdx);
+	for (int32 AnimSetIdx = 0; AnimSetIdx < NumSets; AnimSetIdx++)
 	{
 		OutBindings.Reset();
-		if(FContextualAnimSceneBindings::TryCreateBindings(*SceneAsset, VariantIdx, Params, OutBindings))
+		if(FContextualAnimSceneBindings::TryCreateBindings(*SceneAsset, SectionIdx, AnimSetIdx, Params, OutBindings))
 		{
 			return true;
 		}
@@ -279,18 +279,19 @@
 	return false;
 }
 
-bool UContextualAnimUtilities::CalculateScenePivotForAlignmentSection(const FContextualAnimAlignmentSectionData& AligmentSectionData, const FContextualAnimSceneBindings& Bindings, FTransform& OutScenePivot)
+
+bool UContextualAnimUtilities::CalculateAnimSetPivot(const FContextualAnimSetPivotDefinition& AnimSetPivotDef, const FContextualAnimSceneBindings& Bindings, FTransform& OutScenePivot)
 {
-	if (const FContextualAnimSceneBinding* Binding = Bindings.FindBindingByRole(AligmentSectionData.Origin))
+	if (const FContextualAnimSceneBinding* Binding = Bindings.FindBindingByRole(AnimSetPivotDef.Origin))
 	{
-		if (AligmentSectionData.bAlongClosestDistance)
+		if (AnimSetPivotDef.bAlongClosestDistance)
 		{
-			if (const FContextualAnimSceneBinding* OtherBinding = Bindings.FindBindingByRole(AligmentSectionData.OtherRole))
+			if (const FContextualAnimSceneBinding* OtherBinding = Bindings.FindBindingByRole(AnimSetPivotDef.OtherRole))
 			{
 				const FTransform T1 = Binding->GetTransform();
 				const FTransform T2 = OtherBinding->GetTransform();
 
-				OutScenePivot.SetLocation(FMath::Lerp<FVector>(T1.GetLocation(), T2.GetLocation(), AligmentSectionData.Weight));
+				OutScenePivot.SetLocation(FMath::Lerp<FVector>(T1.GetLocation(), T2.GetLocation(), AnimSetPivotDef.Weight));
 				OutScenePivot.SetRotation((T2.GetLocation() - T1.GetLocation()).GetSafeNormal2D().ToOrientationQuat());
 				return true;
 			}
@@ -323,22 +324,21 @@
 	const UContextualAnimSceneAsset* SceneAsset = Bindings.GetSceneAsset();
 	if(ensureAlways(SceneAsset))
 	{
-		for (const FContextualAnimAlignmentSectionData& AlignmentSection : SceneAsset->GetAlignmentSections())
+		const int32 SectionIdx = Bindings.GetSectionIdx();
+		for (const FContextualAnimSetPivotDefinition& PivotDef : SceneAsset->GetAnimSetPivotDefinitionsInSection(SectionIdx))
 		{
 			FTransform ScenePivot = FTransform::Identity;
-			if (UContextualAnimUtilities::CalculateScenePivotForAlignmentSection(AlignmentSection, Bindings, ScenePivot))
+			if (UContextualAnimUtilities::CalculateAnimSetPivot(PivotDef, Bindings, ScenePivot))
 			{
 				for (const FContextualAnimSceneBinding& Binding : Bindings)
 				{
-					//@TODO: Cache this
-					const float Time = Binding.GetAnimTrack().GetSyncTimeForWarpSection(AlignmentSection.WarpTargetName);
-
-					const FTransform TransformRelativeToScenePivot = Binding.GetAnimTrack().AlignmentData.ExtractTransformAtTime(AlignmentSection.WarpTargetName, Time);
-					const FTransform WarpTarget = (TransformRelativeToScenePivot * ScenePivot);
-
 					if (UMotionWarpingComponent* MotionWarpComp = Binding.GetActor()->FindComponentByClass<UMotionWarpingComponent>())
 					{
-						MotionWarpComp->AddOrUpdateWarpTargetFromTransform(AlignmentSection.WarpTargetName, WarpTarget);
+						//@TODO: Cache this
+						const float Time = Binding.GetAnimTrack().GetSyncTimeForWarpSection(PivotDef.Name);
+						const FTransform TransformRelativeToScenePivot = Binding.GetAnimTrack().AlignmentData.ExtractTransformAtTime(PivotDef.Name, Time);
+						const FTransform WarpTarget = (TransformRelativeToScenePivot * ScenePivot);
+						MotionWarpComp->AddOrUpdateWarpTargetFromTransform(PivotDef.Name, WarpTarget);
 					}
 				}
 			}
@@ -352,7 +352,7 @@
 
 	if(const UContextualAnimSceneAsset* SceneAsset = Bindings.GetSceneAsset())
 	{
-		Result = SceneAsset->GetAlignmentTransformForRoleRelativeToOtherRole(Role, RelativeToRole, Bindings.GetVariantIdx(), Time);
+		Result = SceneAsset->GetAlignmentTransformForRoleRelativeToOtherRoleInSection(Bindings.GetSectionIdx(), Bindings.GetAnimSetIdx(), Role, RelativeToRole, Time);
 	}
 
 	return Result;
diff --git a/Engine/Plugins/Experimental/Animation/ContextualAnimation/Source/ContextualAnimation/Public/ContextualAnimManager.h b/Engine/Plugins/Experimental/Animation/ContextualAnimation/Source/ContextualAnimation/Public/ContextualAnimManager.h
--- a/Engine/Plugins/Experimental/Animation/ContextualAnimation/Source/ContextualAnimation/Public/ContextualAnimManager.h
+++ b/Engine/Plugins/Experimental/Animation/ContextualAnimation/Source/ContextualAnimation/Public/ContextualAnimManager.h
@@ -44,7 +44,7 @@
 	/** Starts an scene instance with the supplied actors for each role ignoring selection criteria */
 	UContextualAnimSceneInstance* ForceStartScene(const UContextualAnimSceneAsset& SceneAsset, const FContextualAnimStartSceneParams& Params);
 	
-	/** Attempts to start an scene instance with the supplied actors using the first valid variant based on selection criteria */
+	/** Attempts to start an scene instance with the supplied actors using the first valid set based on selection criteria */
 	UContextualAnimSceneInstance* TryStartScene(const UContextualAnimSceneAsset& SceneAsset, const FContextualAnimStartSceneParams& Params);
 
 	UFUNCTION(BlueprintCallable, Category = "Contextual Anim|Manager")
@@ -61,7 +61,7 @@
 
 	FORCEINLINE const TSet<UContextualAnimSceneActorComponent*>& GetSceneActorCompContainer() const { return SceneActorCompContainer; };
 
-	/** Attempts to start an scene instance with the supplied actors using the first valid variant based on selection criteria */
+	/** Attempts to start an scene instance with the supplied actors using the first valid set based on selection criteria */
 	UFUNCTION(BlueprintCallable, Category = "Contextual Anim|Manager", meta = (DisplayName = "Try Start Scene"))
 	UContextualAnimSceneInstance* BP_TryStartScene(const UContextualAnimSceneAsset* SceneAsset, const FContextualAnimStartSceneParams& Params);
 
diff --git a/Engine/Plugins/Experimental/Animation/ContextualAnimation/Source/ContextualAnimation/Public/ContextualAnimSceneAsset.h b/Engine/Plugins/Experimental/Animation/ContextualAnimation/Source/ContextualAnimation/Public/ContextualAnimSceneAsset.h
--- a/Engine/Plugins/Experimental/Animation/ContextualAnimation/Source/ContextualAnimation/Public/ContextualAnimSceneAsset.h
+++ b/Engine/Plugins/Experimental/Animation/ContextualAnimation/Source/ContextualAnimation/Public/ContextualAnimSceneAsset.h
@@ -9,6 +9,7 @@
 
 class UContextualAnimScenePivotProvider;
 class UContextualAnimSceneInstance;
+class UContextualAnimSceneAsset;
 
 UCLASS(Blueprintable)
 class CONTEXTUALANIMATION_API UContextualAnimRolesAsset : public UDataAsset
@@ -30,16 +31,70 @@
 	FORCEINLINE int32 GetNumRoles() const { return Roles.Num(); }
 };
 
+/** List of AnimTrack for each role */
 USTRUCT(BlueprintType)
-struct FContextualAnimTracksContainer
+struct CONTEXTUALANIMATION_API FContextualAnimSet
 {
 	GENERATED_BODY()
 
-	UPROPERTY(EditAnywhere, BlueprintReadWrite, Category = "Defaults")
+	UPROPERTY(EditAnywhere, Category = "Defaults")
 	TArray<FContextualAnimTrack> Tracks;
 
-	UPROPERTY(EditAnywhere, BlueprintReadWrite, Category = "Defaults")
+	UPROPERTY(EditAnywhere, Category = "Defaults")
 	TArray<FTransform> ScenePivots;
+
+	UPROPERTY(VisibleAnywhere, Category = "Settings")
+	float Radius = 0.f;
+};
+
+/** Named container with one or more ContextualAnimSet */
+USTRUCT(BlueprintType)
+struct CONTEXTUALANIMATION_API FContextualAnimSceneSection
+{
+	GENERATED_BODY()
+
+public:
+
+	const FContextualAnimSet* GetAnimSet(int32 AnimSetIdx) const; 
+
+	const FContextualAnimTrack* GetAnimTrack(int32 AnimSetIdx, const FName& Role) const;
+
+	const FContextualAnimTrack* GetAnimTrack(int32 AnimSetIdx, int32 AnimTrackIdx) const;
+
+	FTransform GetAlignmentTransformForRoleRelativeToPivot(int32 AnimSetIdx, FName Role, float Time) const;
+
+	FTransform GetAlignmentTransformForRoleRelativeToOtherRole(int32 AnimSetIdx, FName Role, FName OtherRole, float Time) const;
+
+	FTransform GetIKTargetTransformForRoleAtTime(int32 AnimSetIdx, FName Role, FName TrackName, float Time) const;
+
+	const FContextualAnimIKTargetDefContainer& GetIKTargetDefsForRole(const FName& Role) const;
+
+	const FContextualAnimTrack* FindFirstAnimTrackForRoleThatPassesSelectionCriteria(const FName& Role, const FContextualAnimSceneBindingContext& Primary, const FContextualAnimSceneBindingContext& Querier) const;
+
+	const FContextualAnimTrack* FindAnimTrackForRoleWithClosestEntryLocation(const FName& Role, const FContextualAnimSceneBindingContext& Primary, const FVector& TestLocation) const;
+
+	FORCEINLINE FName GetName() const { return Name; }
+	FORCEINLINE const TArray<FContextualAnimSetPivotDefinition>& GetAnimSetPivotDefinitions() const { return AnimSetPivotDefinitions; }
+
+protected:
+
+	UPROPERTY(EditAnywhere, Category = "Defaults")
+	FName Name;
+
+	UPROPERTY(EditAnywhere, Category = "Defaults")
+	TArray<FContextualAnimSet> AnimSets;
+
+	UPROPERTY(EditAnywhere, Category = "Defaults")
+	TMap<FName, FContextualAnimIKTargetDefContainer> RoleToIKTargetDefsMap;
+
+	UPROPERTY(EditAnywhere, Category = "Defaults")
+	TArray<FContextualAnimSetPivotDefinition> AnimSetPivotDefinitions;
+
+	void GenerateAlignmentTracks(UContextualAnimSceneAsset& SceneAsset);
+	void GenerateIKTargetTracks(UContextualAnimSceneAsset& SceneAsset);
+
+ 	friend class UContextualAnimSceneAsset;
+ 	friend class FContextualAnimViewModel;
 };
 
 UCLASS(Blueprintable)
@@ -53,26 +108,20 @@
 
 	UContextualAnimSceneAsset(const FObjectInitializer& ObjectInitializer);
 
-#if WITH_EDITOR
-	virtual void PostEditChangeProperty(FPropertyChangedEvent& PropertyChangedEvent) override;
-#endif
-
 	virtual void PreSave(FObjectPreSaveContext ObjectSaveContext) override;
 
 	void PrecomputeData();
 	
 	void ForEachAnimTrack(FForEachAnimTrackFunction Function) const;
 
-	void ForEachAnimTrack(int32 VariantIdx, FForEachAnimTrackFunction Function) const;
-
 	FORCEINLINE const FName& GetPrimaryRole() const { return PrimaryRole; }
-	FORCEINLINE float GetRadius() const { return Radius; }
+	
 	FORCEINLINE bool GetDisableCollisionBetweenActors() const { return bDisableCollisionBetweenActors; }
 	FORCEINLINE const TSubclassOf<UContextualAnimSceneInstance>& GetSceneInstanceClass() const { return SceneInstanceClass; }
-	FORCEINLINE const TArray<FContextualAnimAlignmentSectionData>& GetAlignmentSections() const { return  AlignmentSections; }
+	FORCEINLINE int32 GetSampleRate() const { return SampleRate; }
+	
+	bool HasValidData() const { return RolesAsset != nullptr && Sections.Num() > 0 && Sections[0].AnimSets.Num() > 0; }
 
-	bool HasValidData() const { return RolesAsset && Variants.Num() > 0; }
-
 	const UContextualAnimRolesAsset* GetRolesAsset() const { return RolesAsset; }
 
 	UFUNCTION()
@@ -80,39 +129,45 @@
 
 	int32 GetNumRoles() const { return RolesAsset ? RolesAsset->GetNumRoles() : 0; }
 
-	const FContextualAnimTrack* GetAnimTrack(const FName& Role, int32 VariantIdx) const;
-
-	UFUNCTION(BlueprintCallable, Category = "Contextual Anim|Scene Asset")
-	int32 FindVariantIdx(FName Role, UAnimSequenceBase* Animation) const;
-
-	FName FindRoleByAnimation(const UAnimSequenceBase* Animation) const;
-
-	const FContextualAnimTrack* FindFirstAnimTrackForRoleThatPassesSelectionCriteria(const FName& Role, const FContextualAnimSceneBindingContext& Primary, const FContextualAnimSceneBindingContext& Querier) const;
-
-	const FContextualAnimTrack* FindAnimTrackForRoleWithClosestEntryLocation(const FName& Role, const FContextualAnimSceneBindingContext& Primary, const FVector& TestLocation) const;
-
-	UFUNCTION(BlueprintCallable, Category = "Contextual Anim|Scene Asset")
-	FTransform GetAlignmentTransformForRoleRelativeToScenePivot(FName Role, int32 VariantIdx, float Time) const;
-
-	UFUNCTION(BlueprintCallable, Category = "Contextual Anim|Scene Asset")
-	FTransform GetAlignmentTransformForRoleRelativeToOtherRole(FName FromRole, FName ToRole, int32 VariantIdx, float Time) const;
-
-	UFUNCTION(BlueprintCallable, Category = "Contextual Anim|Scene Asset")
-	FTransform GetIKTargetTransformForRoleAtTime(FName Role, int32 VariantIdx, FName TrackName, float Time) const;
-
-	const FContextualAnimIKTargetDefContainer& GetIKTargetDefsForRole(const FName& Role) const;
-
 	const FTransform& GetMeshToComponentForRole(const FName& Role) const;
 
-	int32 GetTotalVariants() const;
+	TArray<FName> GetSectionNames() const;
 
-	UFUNCTION(BlueprintCallable, Category = "Contextual Anim|Scene Asset", meta = (DisplayName = "Get Anim Track"))
-	const FContextualAnimTrack& BP_GetAnimTrack(FName Role, int32 VariantIdx) const
-	{
-		const FContextualAnimTrack* AnimTrack = GetAnimTrack(Role, VariantIdx);
-		return AnimTrack ? *AnimTrack : FContextualAnimTrack::EmptyTrack;
-	}
+	int32 GetNumSections() const;
 
+	int32 GetNumAnimSetsInSection(int32 SectionIdx) const;
+
+	const FContextualAnimSceneSection* GetSection(int32 SectionIdx) const;
+
+	const FContextualAnimSceneSection* GetSection(const FName& SectionName) const;
+
+ 	const FContextualAnimTrack* GetAnimTrack(int32 SectionIdx, int32 AnimSetIdx, const FName& Role) const;
+ 
+ 	const FContextualAnimTrack* GetAnimTrack(int32 SectionIdx, int32 AnimSetIdx, int32 AnimTrackIdx) const;
+
+	const FContextualAnimTrack* FindAnimTrackByAnimation(const UAnimSequenceBase* Animation) const;
+
+	const TArray<FContextualAnimSetPivotDefinition>& GetAnimSetPivotDefinitionsInSection(int32 SectionIdx) const;
+
+	FTransform GetAlignmentTransformForRoleRelativeToOtherRoleInSection(int32 SectionIdx, int32 AnimSetIdx, const FName& Role, const FName& OtherRole, float Time) const;
+
+	const FContextualAnimIKTargetDefContainer& GetIKTargetDefsForRoleInSection(int32 SectionIdx, const FName& Role) const;
+
+	// Blueprint Interface
+	//------------------------------------------------------------------------------------------
+
+	UFUNCTION(BlueprintCallable, Category = "Contextual Anim|Scene Asset", meta = (DisplayName = "Find Animation For Role"))
+	UAnimSequenceBase* BP_FindAnimationForRole(int32 SectionIdx, int32 AnimSetIdx, FName Role) const;
+
+	UFUNCTION(BlueprintCallable, Category = "Contextual Anim|Scene Asset", meta = (DisplayName = "Find AnimSet Index By Animation"))
+	int32 BP_FindAnimSetIndexByAnimation(int32 SectionIdx, const UAnimSequenceBase* Animation) const;
+
+	UFUNCTION(BlueprintCallable, Category = "Contextual Anim|Scene Asset", meta = (DisplayName = "Get Alignment Transform For Role Relative To Pivot"))
+	FTransform BP_GetAlignmentTransformForRoleRelativeToPivot(int32 SectionIdx, int32 AnimSetIdx, FName Role, float Time) const;
+
+	UFUNCTION(BlueprintCallable, Category = "Contextual Anim|Scene Asset", meta = (DisplayName = "Get IK Target Transform For Role At Time"))
+	FTransform BP_GetIKTargetTransformForRoleAtTime(int32 SectionIdx, int32 AnimSetIdx, FName Role, FName TrackName, float Time) const;
+
 	//@TODO: Kept around only to do not break existing content. It will go away in the future.
 	UFUNCTION(BlueprintCallable, Category = "Contextual Anim|Scene Asset")
 	bool Query(FName Role, FContextualAnimQueryResult& OutResult, const FContextualAnimQueryParams& QueryParams, const FTransform& ToWorldTransform) const;
@@ -126,13 +181,7 @@
 	FName PrimaryRole = NAME_None;
 
 	UPROPERTY(EditAnywhere, Category = "Defaults")
-	TArray<FContextualAnimTracksContainer> Variants;
-
-	UPROPERTY(EditAnywhere, Category = "Defaults")
-	TMap<FName, FContextualAnimIKTargetDefContainer> RoleToIKTargetDefsMap;
-
-	UPROPERTY(EditAnywhere, Category = "Defaults")
-	TArray<FContextualAnimAlignmentSectionData> AlignmentSections;
+	TArray<FContextualAnimSceneSection> Sections;
 
 	UPROPERTY(EditAnywhere, Category = "Settings")
 	TSubclassOf<UContextualAnimSceneInstance> SceneInstanceClass;
@@ -144,17 +193,5 @@
 	UPROPERTY(EditAnywhere, Category = "Settings", meta = (ClampMin = "1", ClampMax = "60"), AdvancedDisplay)
 	int32 SampleRate = 15;
 
-	UPROPERTY(VisibleAnywhere, Category = "Settings")
-	float Radius = 0.f;
-
-	void GenerateAlignmentTracks();
-
-	void GenerateIKTargetTracks();
-
-	void UpdateRadius();
-
-	friend class UContextualAnimUtilities;
-	friend class FContextualAnimViewModel;
-	friend class FContextualAnimEdMode;
-	friend class FContextualAnimMovieSceneNotifyTrackEditor;
+ 	friend class FContextualAnimViewModel;
 };
diff --git a/Engine/Plugins/Experimental/Animation/ContextualAnimation/Source/ContextualAnimation/Public/ContextualAnimSceneInstance.h b/Engine/Plugins/Experimental/Animation/ContextualAnimation/Source/ContextualAnimation/Public/ContextualAnimSceneInstance.h
--- a/Engine/Plugins/Experimental/Animation/ContextualAnimation/Source/ContextualAnimation/Public/ContextualAnimSceneInstance.h
+++ b/Engine/Plugins/Experimental/Animation/ContextualAnimation/Source/ContextualAnimation/Public/ContextualAnimSceneInstance.h
@@ -83,10 +83,6 @@
 	const FContextualAnimSceneBinding* FindBindingByActor(const AActor* Actor) const { return Bindings.FindBindingByActor(Actor); }
 	const FContextualAnimSceneBinding* FindBindingByRole(const FName& Role) const { return Bindings.FindBindingByRole(Role); }
 
-#if WITH_EDITOR
-	const FContextualAnimSceneBinding* FindBindingByGuid(const FGuid& Guid) const { return Bindings.FindBindingByGuid(Guid); }
-#endif
-
 	UFUNCTION(BlueprintCallable, Category = "Contextual Anim|Scene Instance")
 	AActor* GetActorByRole(FName Role) const;
 
@@ -98,7 +94,7 @@
 	/** Tells the scene actor to leave the scene (stop animation) */
 	void Leave(FContextualAnimSceneBinding& Binding);
 
-	bool TransitionTo(FContextualAnimSceneBinding& Binding, const FName& ToSectionName);
+	bool TransitionTo(FContextualAnimSceneBinding& Binding, const FContextualAnimTrack& AnimTrack);
 
 	/** Helper function to set ignore collision between the supplied actor and all the other actors in this scene */
 	void SetIgnoreCollisionWithOtherActors(AActor* Actor, bool bValue) const;
@@ -122,4 +118,7 @@
 	FContextualAnimSceneBindings Bindings;
 
 	TArray<TTuple<FName, FTransform>> AlignmentSectionToScenePivotList;
+
+	/** Helper to play an AnimSequenceBase as montage. If Animation is not a montage it plays it as dynamic montage  */
+	UAnimMontage* PlayAnimation(UAnimInstance& AnimInstance, UAnimSequenceBase& Animation);
 };
diff --git a/Engine/Plugins/Experimental/Animation/ContextualAnimation/Source/ContextualAnimation/Public/ContextualAnimTypes.h b/Engine/Plugins/Experimental/Animation/ContextualAnimation/Source/ContextualAnimation/Public/ContextualAnimTypes.h
--- a/Engine/Plugins/Experimental/Animation/ContextualAnimation/Source/ContextualAnimation/Public/ContextualAnimTypes.h
+++ b/Engine/Plugins/Experimental/Animation/ContextualAnimation/Source/ContextualAnimation/Public/ContextualAnimTypes.h
@@ -91,8 +91,11 @@
 	FName Role = NAME_None;
 
 	UPROPERTY(VisibleAnywhere, BlueprintReadOnly, Category = "Defaults")
-	int32 VariantIdx = INDEX_NONE;
+	int32 SectionIdx = INDEX_NONE;
 
+	UPROPERTY(VisibleAnywhere, BlueprintReadOnly, Category = "Defaults")
+	int32 AnimSetIdx = INDEX_NONE;
+
 	float GetSyncTimeForWarpSection(int32 WarpSectionIndex) const;
 	float GetSyncTimeForWarpSection(const FName& WarpSectionName) const;
 
@@ -104,6 +107,8 @@
 
 	bool DoesQuerierPassSelectionCriteria(const FContextualAnimSceneBindingContext& PrimaryActorData, const FContextualAnimSceneBindingContext& QuerierData) const;
 
+	FTransform GetRootTransformAtTime(float Time) const;
+
 	static const FContextualAnimTrack EmptyTrack;
 };
 
@@ -210,16 +215,17 @@
 	FTransform MeshToComponent = FTransform(FRotator(0.f, -90.f, 0.f));
 };
 
-// FContextualAnimAlignmentSectionData
+// FContextualAnimSetPivotDefinition
 ///////////////////////////////////////////////////////////////////////
 
+/** Rules used to compute the pivot for a AnimSet */
 USTRUCT(BlueprintType)
-struct FContextualAnimAlignmentSectionData
+struct FContextualAnimSetPivotDefinition
 {
 	GENERATED_BODY()
 
 	UPROPERTY(EditAnywhere, BlueprintReadWrite, Category = "Defaults")
-	FName WarpTargetName = NAME_None;
+	FName Name = NAME_None;
 
 	UPROPERTY(EditAnywhere, BlueprintReadWrite, Category = "Defaults", meta = (GetOptions = "GetRoles"))
 	FName Origin = NAME_None;
@@ -288,6 +294,7 @@
 
 	/** Returns the track in the scene asset used by this actor */
 	FORCEINLINE const FContextualAnimTrack& GetAnimTrack() const { check(AnimTrackPtr);  return *AnimTrackPtr; }
+	FORCEINLINE void SetAnimTrack(const FContextualAnimTrack& InAnimTrack) { AnimTrackPtr = &InAnimTrack; }
 
 	/** Returns the role definition used by this actor */
 	FORCEINLINE const FContextualAnimRoleDefinition& GetRoleDef() const { check(RoleDefPtr); return *RoleDefPtr; }
@@ -335,13 +342,6 @@
 
 	/** Ptr back to the scene instance we belong to (if any) */
 	TWeakObjectPtr<const UContextualAnimSceneInstance> SceneInstancePtr = nullptr;
-
-#if WITH_EDITOR
-public:
-	/** Guid only used in editor to bind this actor to sequencer */
-	FGuid Guid;
-#endif
-
 };
 
 USTRUCT(BlueprintType)
@@ -359,15 +359,9 @@
 		return Role != NAME_None ? Data.FindByPredicate([&Role](const auto& Item) { return Item.GetAnimTrack().Role == Role; }) : nullptr;
 	}
 
-#if WITH_EDITOR
-	const FContextualAnimSceneBinding* FindBindingByGuid(const FGuid& Guid) const
-	{
-		return Guid.IsValid() ? Data.FindByPredicate([&Guid](const auto& Item) { return Item.Guid == Guid; }) : nullptr;
-	}
-#endif
-
 	FORCEINLINE const UContextualAnimSceneAsset* GetSceneAsset() const { return Num() > 0 ? &Data[0].GetSceneAsset() : nullptr; }
-	FORCEINLINE int32 GetVariantIdx() const { return Num() > 0 ? Data[0].GetAnimTrack().VariantIdx : INDEX_NONE; }
+	FORCEINLINE int32 GetSectionIdx() const { return Num() > 0 ? Data[0].GetAnimTrack().SectionIdx : INDEX_NONE; }
+	FORCEINLINE int32 GetAnimSetIdx() const { return Num() > 0 ? Data[0].GetAnimTrack().AnimSetIdx : INDEX_NONE; }
 	FORCEINLINE int32 Num() const { return Data.Num(); }
 	FORCEINLINE int32 Add(const FContextualAnimSceneBinding& NewData) { return Data.Add(NewData); }
 	FORCEINLINE void Reset() { return Data.Reset(); }
@@ -378,7 +372,7 @@
 	FORCEINLINE TArray<FContextualAnimSceneBinding>::RangedForIteratorType      end() { return Data.end(); }
 	FORCEINLINE TArray<FContextualAnimSceneBinding>::RangedForConstIteratorType end() const { return Data.end(); }
 
-	static bool TryCreateBindings(const UContextualAnimSceneAsset& SceneAsset, int32 VariantIdx, const TMap<FName, FContextualAnimSceneBindingContext>& Params, FContextualAnimSceneBindings& OutBindings);
+	static bool TryCreateBindings(const UContextualAnimSceneAsset& SceneAsset, int32 SectionIdx, int32 AnimSetIdx, const TMap<FName, FContextualAnimSceneBindingContext>& Params, FContextualAnimSceneBindings& OutBindings);
 
 private:
 
@@ -396,20 +390,23 @@
 	UPROPERTY(EditAnywhere, BlueprintReadWrite, Category = "Defaults")
 	TMap<FName, FContextualAnimSceneBindingContext> RoleToActorMap;
 
-	/** Desired variant. If INDEX_NONE the Manager will attempt to find the best variant to use by running the selection criteria */
+	/** Desired set. If INDEX_NONE the Manager will attempt to find the best set to use (in the first section) by running the selection criteria */
 	UPROPERTY(EditAnywhere, BlueprintReadWrite, Category = "Defaults")
-	int32 VariantIdx = INDEX_NONE;
+	int32 AnimSetIdx = INDEX_NONE;
 
 	void Reset()
 	{
 		RoleToActorMap.Reset();
-		VariantIdx = INDEX_NONE;
+		AnimSetIdx = INDEX_NONE;
 	}
 };
 
 ///////////////////////////////////////////////////////////////////////
 
-/** Stores the result of a query function */
+/** 
+ * Stores the result of a query function 
+ * @TODO: Only used by UContextualAnimSceneAsset::Query. Kept around only to do not break existing content. It will go away in the future.
+ */
 USTRUCT(BlueprintType)
 struct FContextualAnimQueryResult
 {
@@ -428,20 +425,23 @@
 	float AnimStartTime = 0.f;
 
 	UPROPERTY(BlueprintReadWrite, Category = "Defaults")
-	int32 VariantIdx = INDEX_NONE;
+	int32 AnimSetIdx = INDEX_NONE;
 
 	void Reset()
 	{
 		Animation.Reset();
 		EntryTransform = SyncTransform = FTransform::Identity;
 		AnimStartTime = 0.f;
-		VariantIdx = INDEX_NONE;
+		AnimSetIdx = INDEX_NONE;
 	}
 
-	FORCEINLINE bool IsValid() const { return VariantIdx != INDEX_NONE; }
+	FORCEINLINE bool IsValid() const { return AnimSetIdx != INDEX_NONE; }
 };
 
-/** Stores the parameters passed into query function */
+/** 
+ * Stores the parameters passed into query function 
+ * @TODO: Only used by UContextualAnimSceneAsset::Query. Kept around only to do not break existing content. It will go away in the future.
+ */
 USTRUCT(BlueprintType)
 struct FContextualAnimQueryParams
 {
diff --git a/Engine/Plugins/Experimental/Animation/ContextualAnimation/Source/ContextualAnimation/Public/ContextualAnimUtilities.h b/Engine/Plugins/Experimental/Animation/ContextualAnimation/Source/ContextualAnimation/Public/ContextualAnimUtilities.h
--- a/Engine/Plugins/Experimental/Animation/ContextualAnimation/Source/ContextualAnimation/Public/ContextualAnimUtilities.h
+++ b/Engine/Plugins/Experimental/Animation/ContextualAnimation/Source/ContextualAnimation/Public/ContextualAnimUtilities.h
@@ -13,6 +13,7 @@
 class AActor;
 class FPrimitiveDrawInterface;
 struct FAnimMontageInstance;
+struct FContextualAnimSet;
 
 UCLASS()
 class CONTEXTUALANIMATION_API UContextualAnimUtilities : public UBlueprintFunctionLibrary
@@ -43,7 +44,7 @@
 
 	static void DrawDebugPose(const UWorld* World, const UAnimSequenceBase* Animation, float Time, const FTransform& LocalToWorldTransform, const FColor& Color, float LifeTime, float Thickness);
 	
-	static void DrawDebugScene(const UWorld* World, const UContextualAnimSceneAsset* SceneAsset, int32 VariantIdx, float Time, const FTransform& ToWorldTransform, const FColor& Color, float LifeTime, float Thickness);
+	static void DrawDebugAnimSet(const UWorld* World, const UContextualAnimSceneAsset& SceneAsset, const FContextualAnimSet& AnimSet, float Time, const FTransform& ToWorldTransform, const FColor& Color, float LifeTime, float Thickness);
 
 	static USkeletalMeshComponent* TryGetSkeletalMeshComponent(const AActor* Actor);
 
@@ -53,7 +54,7 @@
 
 	static void DrawSector(FPrimitiveDrawInterface& PDI, const FVector& Origin, const FVector& Direction, float MinDistance, float MaxDistance, float MinAngle, float MaxAngle, const FLinearColor& Color, uint8 DepthPriority, float Thickness);
 
-	static bool CalculateScenePivotForAlignmentSection(const FContextualAnimAlignmentSectionData& AligmentSectionData, const FContextualAnimSceneBindings& Bindings, FTransform& OutScenePivot);
+	static bool CalculateAnimSetPivot(const FContextualAnimSetPivotDefinition& AnimSetPivotDef, const FContextualAnimSceneBindings& Bindings, FTransform& OutScenePivot);
 
 	UFUNCTION(BlueprintCallable, Category = "Contextual Anim|Scene Asset", meta = (DisplayName = "Create Contextual Anim Scene Bindings"))
 	static bool BP_CreateContextualAnimSceneBindings(const UContextualAnimSceneAsset* SceneAsset, const TMap<FName, FContextualAnimSceneBindingContext>& Params, FContextualAnimSceneBindings& OutBindings);
diff --git a/Engine/Plugins/Experimental/Animation/ContextualAnimation/Source/ContextualAnimationEditor/Private/ContextualAnimAssetEditorToolkit.cpp b/Engine/Plugins/Experimental/Animation/ContextualAnimation/Source/ContextualAnimationEditor/Private/ContextualAnimAssetEditorToolkit.cpp
--- a/Engine/Plugins/Experimental/Animation/ContextualAnimation/Source/ContextualAnimationEditor/Private/ContextualAnimAssetEditorToolkit.cpp
+++ b/Engine/Plugins/Experimental/Animation/ContextualAnimation/Source/ContextualAnimationEditor/Private/ContextualAnimAssetEditorToolkit.cpp
@@ -203,10 +203,10 @@
 
 	ToolbarBuilder.AddComboButton(
 		FUIAction(),
-		FOnGetContent::CreateSP(this, &FContextualAnimAssetEditorToolkit::BuildVariantsMenu),
-		LOCTEXT("Variants_Label", "Variants"),
+		FOnGetContent::CreateSP(this, &FContextualAnimAssetEditorToolkit::BuildNewAnimSetWidget),
+		LOCTEXT("NewSet_Label", "New Set"),
 		FText::GetEmpty(),
-		FSlateIcon()
+		FSlateIcon(FAppStyle::GetAppStyleSetName(), "Icons.Plus")
 	);
 
 	ToolbarBuilder.AddToolBarButton(
@@ -218,102 +218,67 @@
 	);
 }
 
-TSharedRef<SWidget> FContextualAnimAssetEditorToolkit::BuildVariantsMenu()
+TSharedRef<SWidget> FContextualAnimAssetEditorToolkit::BuildNewAnimSetWidget()
 {
 	const bool bShouldCloseWindowAfterMenuSelection = true;
 	FMenuBuilder MenuBuilder(bShouldCloseWindowAfterMenuSelection, GetToolkitCommands());
 
+	FPropertyEditorModule& PropertyModule = FModuleManager::LoadModuleChecked<FPropertyEditorModule>("PropertyEditor");
+
+	FDetailsViewArgs Args;
+	Args.bHideSelectionTip = true;
+	Args.bAllowSearch = false;
+	Args.bAllowFavoriteSystem = false;
+
+	NewAnimSetWidgetStruct = MakeShared<FStructOnScope>(FContextualAnimNewAnimSetParams::StaticStruct());
+	FContextualAnimNewAnimSetParams* Params = (FContextualAnimNewAnimSetParams*)NewAnimSetWidgetStruct->GetStructMemory();
+
+	const TArray<FName> Roles = ViewModel->GetSceneAsset()->GetRoles();
+	for (FName Role : Roles)
 	{
-		MenuBuilder.BeginSection(NAME_None, LOCTEXT("CreateNewVariant", "New Variant"));
+		FContextualAnimNewAnimSetData Entry;
+		Entry.RoleName = Role;
+		Params->Data.Add(Entry);
+	}
 
-		MenuBuilder.AddSubMenu(
-			LOCTEXT("NewVariant", "New Variant"),
-			FText::GetEmpty(),
-			FNewMenuDelegate::CreateLambda([this](FMenuBuilder& MenuBuilder){
-			
-				FPropertyEditorModule& PropertyModule = FModuleManager::LoadModuleChecked<FPropertyEditorModule>("PropertyEditor");
+	TSharedRef<IStructureDetailsView> StructureDetailsView = PropertyModule.CreateStructureDetailView(Args, FStructureDetailsViewArgs(), NewAnimSetWidgetStruct);
 
-				FDetailsViewArgs Args;
-				Args.bHideSelectionTip = true;
-				Args.bAllowSearch = false;
-				Args.bAllowFavoriteSystem = false;
-
-				NewVariantWidgetStruct = MakeShared<FStructOnScope>(FContextualAnimNewVariantParams::StaticStruct());
-				FContextualAnimNewVariantParams* Params = (FContextualAnimNewVariantParams*)NewVariantWidgetStruct->GetStructMemory();
-
-				const TArray<FName> Roles = ViewModel->GetSceneAsset()->GetRoles();
-				for (FName Role : Roles)
+	MenuBuilder.AddWidget(
+		SNew(SBox)
+		.MinDesiredWidth(500.0f)
+		.MaxDesiredWidth(500.f)
+		.MaxDesiredHeight(400.0f)
+		[
+			SNew(SVerticalBox)
+			+ SVerticalBox::Slot()
+			.AutoHeight()
+			.HAlign(HAlign_Fill)
+			[
+				StructureDetailsView->GetWidget().ToSharedRef()
+			]
+			+ SVerticalBox::Slot()
+			.AutoHeight()
+			.HAlign(HAlign_Fill)
+			.Padding(5)
+			[
+				SNew(SButton)
+				.HAlign(HAlign_Center)
+				.ContentPadding(FAppStyle::GetMargin("StandardDialog.ContentPadding"))
+				.Text(LOCTEXT("OK", "OK"))
+				.OnClicked_Lambda([this]()
 				{
-					FContextualAnimNewVariantData Entry;
-					Entry.RoleName = Role;
-					Params->Data.Add(Entry);
-				}
+					FContextualAnimNewAnimSetParams* Params = (FContextualAnimNewAnimSetParams*)NewAnimSetWidgetStruct->GetStructMemory();
+					check(Params);
 
-				TSharedRef<IStructureDetailsView> StructureDetailsView = PropertyModule.CreateStructureDetailView(Args, FStructureDetailsViewArgs(), NewVariantWidgetStruct);
+					ViewModel->AddNewAnimSet(*Params);
 
-				MenuBuilder.AddWidget(
-					SNew(SBox)
-					.MinDesiredWidth(500.0f)
-					.MaxDesiredWidth(500.f)
-					.MaxDesiredHeight(400.0f)
-					[
-						SNew(SVerticalBox)
-						+ SVerticalBox::Slot()
-						.AutoHeight()
-						.HAlign(HAlign_Fill)
-						[
-							StructureDetailsView->GetWidget().ToSharedRef()
-						]
-						+ SVerticalBox::Slot()
-						.AutoHeight()
-						.HAlign(HAlign_Fill)
-						[
-							SNew(SButton)
-							.ContentPadding(3)
-							.VAlign(VAlign_Center)
-							.HAlign(HAlign_Center)
-							.OnClicked_Lambda([this]()
-							{
-								FContextualAnimNewVariantParams* Params = (FContextualAnimNewVariantParams*)NewVariantWidgetStruct->GetStructMemory();
-								check(Params);
+					FSlateApplication::Get().DismissAllMenus();
 
-								ViewModel->AddNewVariant(*Params);
+					return FReply::Handled();
+				})
+			]
+		], FText(), true, false);
 
-								FSlateApplication::Get().DismissAllMenus();
-
-								return FReply::Handled();
-							})
-						.Text(LOCTEXT("OK", "OK"))
-						]
-					],
-					FText(), true, false);
-			}),
-			false,
-			FSlateIcon()
-		);
-
-		MenuBuilder.EndSection();
-	}
-
-	{
-		MenuBuilder.BeginSection(NAME_None, LOCTEXT("Variants_Label", "Variants"));
-		{
-			const UContextualAnimSceneAsset* SceneAsset = GetSceneAsset();
-			const int32 TotalVariants = SceneAsset->GetTotalVariants();
-			for (int32 Idx = 0; Idx < TotalVariants; Idx++)
-			{
-				MenuBuilder.AddMenuEntry(
-					FText::FromString(FString::Printf(TEXT("%d"), Idx)),
-					FText::GetEmpty(),
-					FSlateIcon(),
-					FUIAction(FExecuteAction::CreateLambda([this, Idx]() {
-						ViewModel->SetActiveSceneVariantIdx(Idx);
-					})));
-			}
-		}
-		MenuBuilder.EndSection();
-	}
-
 	return MenuBuilder.MakeWidget();
 }
 
diff --git a/Engine/Plugins/Experimental/Animation/ContextualAnimation/Source/ContextualAnimationEditor/Private/ContextualAnimEdMode.cpp b/Engine/Plugins/Experimental/Animation/ContextualAnimation/Source/ContextualAnimationEditor/Private/ContextualAnimEdMode.cpp
--- a/Engine/Plugins/Experimental/Animation/ContextualAnimation/Source/ContextualAnimationEditor/Private/ContextualAnimEdMode.cpp
+++ b/Engine/Plugins/Experimental/Animation/ContextualAnimation/Source/ContextualAnimationEditor/Private/ContextualAnimEdMode.cpp
@@ -86,27 +86,24 @@
 
 		if (const UContextualAnimSceneInstance* SceneInstance = ViewModel->GetSceneInstance())
 		{
-			const UContextualAnimSceneAsset* SceneAsset = ViewModel->GetSceneAsset();
-			const int32 VariantIdx = ViewModel->GetActiveSceneVariantIdx();
+			const UContextualAnimSceneAsset& SceneAsset = SceneInstance->GetSceneAsset();
+			const FContextualAnimSceneBindings& Bindings = SceneInstance->GetBindings();
 
 			// Draw Scene Pivots
-			if (SceneAsset->Variants.IsValidIndex(VariantIdx))
+			if(const FContextualAnimSceneSection* Section = SceneAsset.GetSection(Bindings.GetSectionIdx()))
 			{
-				for(const FTransform& ScenePivot : SceneAsset->Variants[VariantIdx].ScenePivots)
+				if(const FContextualAnimSet* AnimSet = Section->GetAnimSet(Bindings.GetAnimSetIdx()))
 				{
-					DrawCoordinateSystem(PDI, ScenePivot.GetLocation(), ScenePivot.Rotator(), 50.f, SDPG_Foreground);
+					for (const FTransform& ScenePivot : AnimSet->ScenePivots)
+					{
+						DrawCoordinateSystem(PDI, ScenePivot.GetLocation(), ScenePivot.Rotator(), 50.f, SDPG_Foreground);
+					}
 				}
 			}
 
-			FTransform PrimaryActorTransform = FTransform::Identity;
-			if (const FContextualAnimSceneBinding* Binding = SceneInstance->FindBindingByRole(SceneAsset->GetPrimaryRole()))
-			{
-				PrimaryActorTransform = Binding->GetTransform();
-			}
-
 			const EShowIKTargetsDrawMode IKTargetsDrawMode = ViewportClient->GetShowIKTargetsDrawMode();
 
-			for (const FContextualAnimSceneBinding& Binding : SceneInstance->GetBindings())
+			for (const FContextualAnimSceneBinding& Binding : Bindings)
 			{
 				// Draw IK Targets
 				if (IKTargetsDrawMode == EShowIKTargetsDrawMode::All || (IKTargetsDrawMode == EShowIKTargetsDrawMode::Selected && Binding.GetActor() == SelectedActor.Get()))
@@ -115,14 +112,17 @@
 				}
 
 				// Draw Selection Criteria
-				if (const FContextualAnimTrack* AnimTrack = SceneAsset->GetAnimTrack(Binding.GetRoleDef().Name, VariantIdx))
+				if(const FContextualAnimSceneBinding* PrimaryBinding = Bindings.FindBindingByRole(SceneAsset.GetPrimaryRole()))
 				{
-					for (int32 CriterionIdx = 0; CriterionIdx < AnimTrack->SelectionCriteria.Num(); CriterionIdx++)
+					const FTransform PrimaryTransform = PrimaryBinding->GetTransform();
+
+					const FContextualAnimTrack& AnimTrack = Binding.GetAnimTrack();
+					for (int32 CriterionIdx = 0; CriterionIdx < AnimTrack.SelectionCriteria.Num(); CriterionIdx++)
 					{
-						if(const UContextualAnimSelectionCriterion* Criterion = AnimTrack->SelectionCriteria[CriterionIdx])
+						if(const UContextualAnimSelectionCriterion* Criterion = AnimTrack.SelectionCriteria[CriterionIdx])
 						{
 							FLinearColor DrawColor = FLinearColor::White;
-							if (Criterion->DoesQuerierPassCondition(FContextualAnimSceneBindingContext(PrimaryActorTransform), Binding.GetContext()))
+							if (Criterion->DoesQuerierPassCondition(FContextualAnimSceneBindingContext(PrimaryTransform), Binding.GetContext()))
 							{
 								DrawColor = FLinearColor::Green;
 							}
@@ -134,8 +134,8 @@
 								const int32 LastIndex = Spatial->PolygonPoints.Num() - 1;
 								for (int32 Idx = 0; Idx <= LastIndex; Idx++)
 								{
-									const FVector P0 = PrimaryActorTransform.TransformPositionNoScale(Spatial->PolygonPoints[Idx]);
-									const FVector P1 = PrimaryActorTransform.TransformPositionNoScale(Spatial->PolygonPoints[Idx == LastIndex ? 0 : Idx + 1]);
+									const FVector P0 = PrimaryTransform.TransformPositionNoScale(Spatial->PolygonPoints[Idx]);
+									const FVector P1 = PrimaryTransform.TransformPositionNoScale(Spatial->PolygonPoints[Idx == LastIndex ? 0 : Idx + 1]);
 
 									PDI->DrawLine(P0, P1, DrawColor, SDPG_Foreground, 2.f);
 
@@ -143,18 +143,18 @@
 
 									PDI->DrawLine(P0, P0 + FVector::UpVector * Spatial->Height, DrawColor, SDPG_Foreground, 2.f);
 
-									PDI->SetHitProxy(new HSelectionCriterionHitProxy(FSelectionCriterionHitProxyData(Binding.GetRoleDef().Name, VariantIdx, CriterionIdx, Idx)));
+									PDI->SetHitProxy(new HSelectionCriterionHitProxy(FSelectionCriterionHitProxyData(AnimTrack.SectionIdx, AnimTrack.AnimSetIdx, AnimTrack.Role, CriterionIdx, Idx)));
 									PDI->DrawPoint(P0, FLinearColor::Black, 15.f, SDPG_Foreground);
 									PDI->SetHitProxy(nullptr);
 
-									PDI->SetHitProxy(new HSelectionCriterionHitProxy(FSelectionCriterionHitProxyData(Binding.GetRoleDef().Name, VariantIdx, CriterionIdx, Idx + 4)));
+									PDI->SetHitProxy(new HSelectionCriterionHitProxy(FSelectionCriterionHitProxyData(AnimTrack.SectionIdx, AnimTrack.AnimSetIdx, AnimTrack.Role, CriterionIdx, Idx + 4)));
 									PDI->DrawPoint(P0 + FVector::UpVector * Spatial->Height, FLinearColor::Black, 15.f, SDPG_Foreground);
 									PDI->SetHitProxy(nullptr);
 								}
 							}
 							else if (const UContextualAnimSelectionCriterion_Facing* Facing = Cast<UContextualAnimSelectionCriterion_Facing>(Criterion))
 							{
-								const FTransform Transform = AnimTrack->GetAlignmentTransformAtEntryTime() * PrimaryActorTransform;
+								const FTransform Transform = AnimTrack.GetAlignmentTransformAtEntryTime() * PrimaryTransform;
 								UContextualAnimUtilities::DrawSector(*PDI, Transform.GetLocation(), Transform.GetRotation().GetForwardVector(), 0.f, 30.f, -Facing->MaxAngle, Facing->MaxAngle, DrawColor, SDPG_World, 1.f);
 							}
 						}
@@ -256,7 +256,7 @@
 			if (SelectedSelectionCriterionData.IsValid())
 			{
 				const UContextualAnimSceneAsset* SceneAsset = ViewModel->GetSceneAsset();
-				if (const FContextualAnimTrack* AnimTrack = SceneAsset->GetAnimTrack(SelectedSelectionCriterionData.RoleName, SelectedSelectionCriterionData.VariantIdx))
+				if (const FContextualAnimTrack* AnimTrack = SceneAsset->GetAnimTrack(SelectedSelectionCriterionData.SectionIdx, SelectedSelectionCriterionData.AnimSetIdx, SelectedSelectionCriterionData.RoleName))
 				{
 					if(AnimTrack->SelectionCriteria.IsValidIndex(SelectedSelectionCriterionData.CriterionIdx))
 					{
@@ -323,7 +323,7 @@
 	if (ViewModel && SelectedSelectionCriterionData.IsValid())
 	{
 		const UContextualAnimSceneAsset* SceneAsset = ViewModel->GetSceneAsset();
-		if (const FContextualAnimTrack* AnimTrack = SceneAsset->GetAnimTrack(SelectedSelectionCriterionData.RoleName, SelectedSelectionCriterionData.VariantIdx))
+		if (const FContextualAnimTrack* AnimTrack = SceneAsset->GetAnimTrack(SelectedSelectionCriterionData.SectionIdx, SelectedSelectionCriterionData.AnimSetIdx, SelectedSelectionCriterionData.RoleName))
 		{
 			const int32 Idx = SelectedSelectionCriterionData.CriterionIdx;
 			if (AnimTrack->SelectionCriteria.IsValidIndex(Idx) && AnimTrack->SelectionCriteria[Idx]->GetClass()->IsChildOf<UContextualAnimSelectionCriterion_TriggerArea>())
@@ -372,7 +372,7 @@
 			if (const UContextualAnimSceneInstance* SceneInstance = ViewModel->GetSceneInstance())
 			{
 				const UContextualAnimSceneAsset* SceneAsset = ViewModel->GetSceneAsset();
-				if (const FContextualAnimTrack* AnimTrack = SceneAsset->GetAnimTrack(SelectedSelectionCriterionData.RoleName, SelectedSelectionCriterionData.VariantIdx))
+				if (const FContextualAnimTrack* AnimTrack = SceneAsset->GetAnimTrack(SelectedSelectionCriterionData.SectionIdx, SelectedSelectionCriterionData.AnimSetIdx, SelectedSelectionCriterionData.RoleName))
 				{
 					if (AnimTrack->SelectionCriteria.IsValidIndex(SelectedSelectionCriterionData.CriterionIdx))
 					{
diff --git a/Engine/Plugins/Experimental/Animation/ContextualAnimation/Source/ContextualAnimationEditor/Private/ContextualAnimEditorTypes.cpp b/Engine/Plugins/Experimental/Animation/ContextualAnimation/Source/ContextualAnimationEditor/Private/ContextualAnimEditorTypes.cpp
--- a/Engine/Plugins/Experimental/Animation/ContextualAnimation/Source/ContextualAnimationEditor/Private/ContextualAnimEditorTypes.cpp
+++ b/Engine/Plugins/Experimental/Animation/ContextualAnimation/Source/ContextualAnimationEditor/Private/ContextualAnimEditorTypes.cpp
@@ -7,14 +7,18 @@
 // UContextualAnimNewIKTargetParams
 ////////////////////////////////////////////////////////////////////////////////////////////////
 
-void UContextualAnimNewIKTargetParams::Reset(const FName& InSourceRole, const UContextualAnimSceneAsset& InSceneAsset)
+void UContextualAnimNewIKTargetParams::Reset(const UContextualAnimSceneAsset& InSceneAsset, const UAnimSequenceBase& InAnimation)
 {
+	const FContextualAnimTrack* AnimTrack = InSceneAsset.FindAnimTrackByAnimation(&InAnimation);
+	check(AnimTrack);
+
 	CachedRoles = InSceneAsset.GetRoles();
-	check(CachedRoles.Contains(InSourceRole));
+	check(CachedRoles.Contains(AnimTrack->Role));
 
 	SceneAssetPtr = &InSceneAsset;
 
-	SourceRole = InSourceRole;
+	SectionIdx = AnimTrack->SectionIdx;
+	SourceRole = AnimTrack->Role;
 	GoalName = NAME_None;
 	TargetBone = FBoneReference();
 	SourceBone = FBoneReference();
@@ -43,20 +47,36 @@
 {
 	bInvalidSkeletonIsError = false;
 
-	const FContextualAnimTrack* AnimTrack = nullptr;
+	FName Role = NAME_None;
 	if (PropertyHandle)
 	{
 		if (PropertyHandle->GetProperty()->GetFName() == GET_MEMBER_NAME_CHECKED(UContextualAnimNewIKTargetParams, SourceBone))
 		{
-			AnimTrack = GetSceneAsset().GetAnimTrack(SourceRole, 0);
+			Role = SourceRole;
 		}
 		else if (PropertyHandle->GetProperty()->GetFName() == GET_MEMBER_NAME_CHECKED(UContextualAnimNewIKTargetParams, TargetBone))
 		{
-			AnimTrack = GetSceneAsset().GetAnimTrack(TargetRole, 0);
+			Role = TargetRole;
 		}
 	}
 
-	return (AnimTrack && AnimTrack->Animation) ? AnimTrack->Animation->GetSkeleton() : nullptr;
+	USkeleton* Skeleton = nullptr;
+	if(Role != NAME_None)
+	{
+		// Iterate over the AnimTracks until find the first track with animation and pull the skeleton from there.
+		GetSceneAsset().ForEachAnimTrack([&Skeleton, Role](const FContextualAnimTrack& AnimTrack)
+			{
+				if (AnimTrack.Role == Role && AnimTrack.Animation)
+				{
+					Skeleton = AnimTrack.Animation->GetSkeleton();
+					return UE::ContextualAnim::EForEachResult::Break;
+				}
+
+				return UE::ContextualAnim::EForEachResult::Continue;
+			});
+	}
+
+	return Skeleton;
 }
 
 TArray<FString> UContextualAnimNewIKTargetParams::GetTargetRoleOptions() const
diff --git a/Engine/Plugins/Experimental/Animation/ContextualAnimation/Source/ContextualAnimationEditor/Private/ContextualAnimMovieSceneNotifyTrackEditor.cpp b/Engine/Plugins/Experimental/Animation/ContextualAnimation/Source/ContextualAnimationEditor/Private/ContextualAnimMovieSceneNotifyTrackEditor.cpp
--- a/Engine/Plugins/Experimental/Animation/ContextualAnimation/Source/ContextualAnimationEditor/Private/ContextualAnimMovieSceneNotifyTrackEditor.cpp
+++ b/Engine/Plugins/Experimental/Animation/ContextualAnimation/Source/ContextualAnimationEditor/Private/ContextualAnimMovieSceneNotifyTrackEditor.cpp
@@ -165,10 +165,10 @@
 		MenuBuilder.BeginSection(NAME_None, LOCTEXT("AvailableIKTargets", "Available IK Targets"));
 
 		const UContextualAnimSceneAsset* SceneAsset = GetMovieSceneSequence().GetViewModel().GetSceneAsset();
-		const FName Role = SceneAsset->FindRoleByAnimation(&Track->GetAnimation());
-		check(Role != NAME_None);
+		const FContextualAnimTrack* AnimTrack = SceneAsset->FindAnimTrackByAnimation(&Track->GetAnimation());
+		check(AnimTrack);
 
-		const FContextualAnimIKTargetDefContainer& IKTargets = SceneAsset->GetIKTargetDefsForRole(Role);
+		const FContextualAnimIKTargetDefContainer& IKTargets = SceneAsset->GetIKTargetDefsForRoleInSection(AnimTrack->SectionIdx, AnimTrack->Role);
 		for (const FContextualAnimIKTargetDefinition& IKTargetDef : IKTargets.IKTargetDefs)
 		{
 			const FName GoalName = IKTargetDef.GoalName;
@@ -207,12 +207,8 @@
 	
 	check(NewIKTargetParams);
 
-	const UContextualAnimSceneAsset* SceneAsset = GetMovieSceneSequence().GetViewModel().GetSceneAsset();
-	const FName SourceRole = SceneAsset->FindRoleByAnimation(&Track->GetAnimation());
-	check(SourceRole != NAME_None);
+	NewIKTargetParams->Reset(*GetMovieSceneSequence().GetViewModel().GetSceneAsset(), Track->GetAnimation());
 
-	NewIKTargetParams->Reset(SourceRole, *SceneAsset);
-
 	FPropertyEditorModule& PropertyModule = FModuleManager::LoadModuleChecked<FPropertyEditorModule>("PropertyEditor");
 
 	FDetailsViewArgs Args;
@@ -247,34 +243,10 @@
 				.IsEnabled_Lambda([this](){ return NewIKTargetParams->HasValidData(); })
 				.OnClicked_Lambda([this, Track, RowIndex]()
 				{
-					UContextualAnimSceneAsset* SceneAsset = GetMovieSceneSequence().GetViewModel().GetSceneAsset();
-					check(SceneAsset);
+					GetMovieSceneSequence().GetViewModel().AddNewIKTarget(*NewIKTargetParams);
 
-					// Add IK Target definition to the scene asset
-					FContextualAnimIKTargetDefinition IKTargetDef;
-					IKTargetDef.GoalName = NewIKTargetParams->GoalName;
-					IKTargetDef.BoneName = NewIKTargetParams->SourceBone.BoneName;
-					IKTargetDef.Provider = NewIKTargetParams->Provider;
-					IKTargetDef.TargetRoleName = NewIKTargetParams->TargetRole;
-					IKTargetDef.TargetBoneName = NewIKTargetParams->TargetBone.BoneName;
-
-					if(FContextualAnimIKTargetDefContainer* ContainerPtr = SceneAsset->RoleToIKTargetDefsMap.Find(NewIKTargetParams->SourceRole))
-					{
-						ContainerPtr->IKTargetDefs.AddUnique(IKTargetDef);
-					}
-					else
-					{
-						FContextualAnimIKTargetDefContainer Container;
-						Container.IKTargetDefs.AddUnique(IKTargetDef);
-						SceneAsset->RoleToIKTargetDefsMap.Add(NewIKTargetParams->SourceRole, Container);
-					}
-
-					SceneAsset->PrecomputeData();
-
-					SceneAsset->MarkPackageDirty();
-
 					// Create IK section for the newly created target
-	 				CreateNewIKSection(Track, RowIndex, IKTargetDef.GoalName);
+	 				CreateNewIKSection(Track, RowIndex, NewIKTargetParams->GoalName);
 
 					FSlateApplication::Get().DismissAllMenus();
 
@@ -327,10 +299,11 @@
 		return NameToTest;
 	};
 
-	for (const FGuid& ObjectBinding : ObjectBindings)
+	// @TODO: Commented out for now until we add the new behavior where the user needs to double-click on the animation to edit the notifies
+	/*for (const FGuid& ObjectBinding : ObjectBindings)
 	{
-		UAnimSequenceBase* Animation = GetMovieSceneSequence().GetViewModel().FindAnimationByGuid(ObjectBinding);
-		check(Animation);
+ 		UAnimSequenceBase* Animation = GetMovieSceneSequence().GetViewModel().FindAnimationByGuid(ObjectBinding);
+ 		check(Animation);
 
 		// Copied from AnimTimelineTrack_Notifies.cpp/FAnimTimelineTrack_Notifies::AddTrack()
 
@@ -349,7 +322,7 @@
 		MovieSceneTrack->Initialize(*Animation, NewNotifyTrack);
 	}
 
-	GetSequencer()->NotifyMovieSceneDataChanged(EMovieSceneDataChangeType::MovieSceneStructureItemAdded);
+	GetSequencer()->NotifyMovieSceneDataChanged(EMovieSceneDataChangeType::MovieSceneStructureItemAdded);*/
 }
 
 bool FContextualAnimMovieSceneNotifyTrackEditor::SupportsType(TSubclassOf<UMovieSceneTrack> Type) const
diff --git a/Engine/Plugins/Experimental/Animation/ContextualAnimation/Source/ContextualAnimationEditor/Private/ContextualAnimMovieSceneSection.cpp b/Engine/Plugins/Experimental/Animation/ContextualAnimation/Source/ContextualAnimationEditor/Private/ContextualAnimMovieSceneSection.cpp
new file mode 100644
--- /dev/null
+++ b/Engine/Plugins/Experimental/Animation/ContextualAnimation/Source/ContextualAnimationEditor/Private/ContextualAnimMovieSceneSection.cpp
@@ -0,0 +1,30 @@
+// Copyright Epic Games, Inc. All Rights Reserved.
+
+#include "ContextualAnimMovieSceneSection.h"
+#include "ContextualAnimMovieSceneTrack.h"
+#include "Animation/AnimSequenceBase.h"
+#include "ContextualAnimViewModel.h"
+#include "ContextualAnimSceneAsset.h"
+
+UContextualAnimMovieSceneTrack& UContextualAnimMovieSceneSection::GetOwnerTrack() const
+{
+	UContextualAnimMovieSceneTrack* OwnerTrackPtr = GetTypedOuter<UContextualAnimMovieSceneTrack>();
+	check(OwnerTrackPtr);
+
+	return *OwnerTrackPtr;
+}
+
+void UContextualAnimMovieSceneSection::Initialize(int32 InSectionIdx, int32 InAnimSetIdx, int32 InAnimTrackIdx)
+{
+	SectionIdx = InSectionIdx;
+	AnimSetIdx = InAnimSetIdx;
+	AnimTrackIdx = InAnimTrackIdx;
+}
+
+FContextualAnimTrack& UContextualAnimMovieSceneSection::GetAnimTrack() const
+{
+	const FContextualAnimTrack* AnimTrack = GetOwnerTrack().GetViewModel().GetSceneAsset()->GetAnimTrack(SectionIdx, AnimSetIdx, AnimTrackIdx);
+	check(AnimTrack);
+
+	return *const_cast<FContextualAnimTrack*>(AnimTrack);
+}
\ No newline at end of file
diff --git a/Engine/Plugins/Experimental/Animation/ContextualAnimation/Source/ContextualAnimationEditor/Private/ContextualAnimMovieSceneSequence.cpp b/Engine/Plugins/Experimental/Animation/ContextualAnimation/Source/ContextualAnimationEditor/Private/ContextualAnimMovieSceneSequence.cpp
--- a/Engine/Plugins/Experimental/Animation/ContextualAnimation/Source/ContextualAnimationEditor/Private/ContextualAnimMovieSceneSequence.cpp
+++ b/Engine/Plugins/Experimental/Animation/ContextualAnimation/Source/ContextualAnimationEditor/Private/ContextualAnimMovieSceneSequence.cpp
@@ -1,10 +1,8 @@
 // Copyright Epic Games, Inc. All Rights Reserved.
 
 #include "ContextualAnimMovieSceneSequence.h"
-#include "Tracks/MovieSceneEventTrack.h"
-#include "Tracks/MovieSceneSkeletalAnimationTrack.h"
-#include "Tracks/MovieSceneActorReferenceTrack.h"
 #include "ContextualAnimMovieSceneNotifyTrack.h"
+#include "ContextualAnimMovieSceneTrack.h"
 #include "ContextualAnimViewModel.h"
 
 UContextualAnimMovieSceneSequence::UContextualAnimMovieSceneSequence(const FObjectInitializer& ObjectInitializer)
@@ -15,11 +13,8 @@
 
 ETrackSupport UContextualAnimMovieSceneSequence::IsTrackSupported(TSubclassOf<class UMovieSceneTrack> InTrackClass) const
 { 
-	//@TODO: Remove UMovieSceneSkeletalAnimationTrack and UMovieSceneEventTrack from here after adding a custom track to represent the animation
-	if (InTrackClass == UMovieSceneSkeletalAnimationTrack::StaticClass() ||
-		InTrackClass == UMovieSceneEventTrack::StaticClass() ||
-		InTrackClass == UMovieSceneActorReferenceTrack::StaticClass() ||
-		InTrackClass == UContextualAnimMovieSceneNotifyTrack::StaticClass())
+	if (InTrackClass == UContextualAnimMovieSceneNotifyTrack::StaticClass() || 
+		InTrackClass == UContextualAnimMovieSceneTrack::StaticClass())
 	{
 		return ETrackSupport::Supported;
 	}
diff --git a/Engine/Plugins/Experimental/Animation/ContextualAnimation/Source/ContextualAnimationEditor/Private/ContextualAnimMovieSceneTrack.cpp b/Engine/Plugins/Experimental/Animation/ContextualAnimation/Source/ContextualAnimationEditor/Private/ContextualAnimMovieSceneTrack.cpp
new file mode 100644
--- /dev/null
+++ b/Engine/Plugins/Experimental/Animation/ContextualAnimation/Source/ContextualAnimationEditor/Private/ContextualAnimMovieSceneTrack.cpp
@@ -0,0 +1,74 @@
+// Copyright Epic Games, Inc. All Rights Reserved.
+
+#include "ContextualAnimMovieSceneTrack.h"
+#include "ContextualAnimMovieSceneSection.h"
+#include "ContextualAnimMovieSceneSequence.h"
+#include "ContextualAnimViewModel.h"
+
+UContextualAnimMovieSceneTrack::UContextualAnimMovieSceneTrack()
+{
+#if WITH_EDITORONLY_DATA
+	TrackTint = FColor(110, 155, 57, 150);
+#endif
+}
+
+void UContextualAnimMovieSceneTrack::Initialize(const FName& InRole)
+{
+	Role = InRole;
+
+	SetDisplayName(FText::FromString(FString::Printf(TEXT("Role: %s"), *Role.ToString())));
+}
+
+FContextualAnimViewModel& UContextualAnimMovieSceneTrack::GetViewModel() const
+{
+	return GetTypedOuter<UMovieScene>()->GetTypedOuter<UContextualAnimMovieSceneSequence>()->GetViewModel();
+}
+
+UMovieSceneSection* UContextualAnimMovieSceneTrack::CreateNewSection()
+{
+	return NewObject<UContextualAnimMovieSceneSection>(this, NAME_None, RF_Transactional);
+}
+
+void UContextualAnimMovieSceneTrack::AddSection(UMovieSceneSection& Section)
+{
+	Sections.Add(&Section);
+}
+
+bool UContextualAnimMovieSceneTrack::SupportsType(TSubclassOf<UMovieSceneSection> SectionClass) const
+{
+	return SectionClass == UContextualAnimMovieSceneSection::StaticClass();;
+}
+
+const TArray<UMovieSceneSection*>& UContextualAnimMovieSceneTrack::GetAllSections() const
+{
+	return Sections;
+}
+
+bool UContextualAnimMovieSceneTrack::HasSection(const UMovieSceneSection& Section) const
+{
+	return Sections.Contains(&Section);
+}
+
+bool UContextualAnimMovieSceneTrack::IsEmpty() const
+{
+	return Sections.Num() != 0;
+}
+
+void UContextualAnimMovieSceneTrack::RemoveSection(UMovieSceneSection& Section)
+{
+	// Remove movie scene section
+	Sections.Remove(&Section);
+}
+
+void UContextualAnimMovieSceneTrack::RemoveSectionAt(int32 SectionIndex)
+{
+	check(Sections.IsValidIndex(SectionIndex));
+	RemoveSection(*Sections[SectionIndex]);
+}
+
+#if WITH_EDITOR
+EMovieSceneSectionMovedResult UContextualAnimMovieSceneTrack::OnSectionMoved(UMovieSceneSection& Section, const FMovieSceneSectionMovedParams& Params)
+{
+	return EMovieSceneSectionMovedResult::None;
+}
+#endif
\ No newline at end of file
diff --git a/Engine/Plugins/Experimental/Animation/ContextualAnimation/Source/ContextualAnimationEditor/Private/ContextualAnimMovieSceneTrackEditor.cpp b/Engine/Plugins/Experimental/Animation/ContextualAnimation/Source/ContextualAnimationEditor/Private/ContextualAnimMovieSceneTrackEditor.cpp
new file mode 100644
--- /dev/null
+++ b/Engine/Plugins/Experimental/Animation/ContextualAnimation/Source/ContextualAnimationEditor/Private/ContextualAnimMovieSceneTrackEditor.cpp
@@ -0,0 +1,98 @@
+// Copyright Epic Games, Inc. All Rights Reserved.
+
+#include "ContextualAnimMovieSceneTrackEditor.h"
+#include "ContextualAnimMovieSceneTrack.h"
+#include "ContextualAnimMovieSceneSection.h"
+#include "ContextualAnimViewModel.h"
+#include "SequencerSectionPainter.h"
+#include "Fonts/FontMeasure.h"
+#include "Styling/AppStyle.h"
+
+#define LOCTEXT_NAMESPACE "FContextualAnimMovieSceneTrackEditor"
+
+// FContextualAnimMovieSceneTrackEditor
+////////////////////////////////////////////////////////////////////////////////////////////////
+
+FContextualAnimMovieSceneTrackEditor::FContextualAnimMovieSceneTrackEditor(TSharedRef<ISequencer> InSequencer)
+	: FMovieSceneTrackEditor(InSequencer)
+{ 
+}
+
+TSharedRef<ISequencerTrackEditor> FContextualAnimMovieSceneTrackEditor::CreateTrackEditor(TSharedRef<ISequencer> InSequencer)
+{
+	return MakeShareable(new FContextualAnimMovieSceneTrackEditor(InSequencer));
+}
+
+TSharedRef<ISequencerSection> FContextualAnimMovieSceneTrackEditor::MakeSectionInterface(UMovieSceneSection& SectionObject, UMovieSceneTrack& Track, FGuid ObjectBinding)
+{
+	checkf(SectionObject.GetClass()->IsChildOf<UContextualAnimMovieSceneSection>(), TEXT("Unsupported section."));
+	return MakeShared<FContextualAnimSection>(SectionObject);
+}
+
+void FContextualAnimMovieSceneTrackEditor::BuildAddTrackMenu(FMenuBuilder& MenuBuilder)
+{
+	// Menu that appears when clicking on the Add Track button next to the Search Tracks bar
+}
+
+TSharedPtr<SWidget> FContextualAnimMovieSceneTrackEditor::BuildOutlinerEditWidget(const FGuid& ObjectBinding, UMovieSceneTrack* Track, const FBuildEditWidgetParams& Params)
+{
+	// [+Section] button on the track
+
+	return SNullWidget::NullWidget;
+}
+
+void FContextualAnimMovieSceneTrackEditor::BuildObjectBindingTrackMenu(FMenuBuilder& MenuBuilder, const TArray<FGuid>& ObjectBindings, const UClass* ObjectClass)
+{	
+	// Builds menu that appears when clicking on the +Track button on an Object Track
+}
+
+bool FContextualAnimMovieSceneTrackEditor::SupportsType(TSubclassOf<UMovieSceneTrack> Type) const
+{
+	return (Type == UContextualAnimMovieSceneTrack::StaticClass());
+}
+
+bool FContextualAnimMovieSceneTrackEditor::SupportsSequence(UMovieSceneSequence* InSequence) const
+{
+	ETrackSupport TrackSupported = InSequence ? InSequence->IsTrackSupported(UContextualAnimMovieSceneTrack::StaticClass()) : ETrackSupport::NotSupported;
+	return TrackSupported == ETrackSupport::Supported;
+}
+
+// FContextualAnimSection
+////////////////////////////////////////////////////////////////////////////////////////////////
+
+FText FContextualAnimSection::GetSectionTitle() const
+{
+	if (UContextualAnimMovieSceneSection* Section = Cast<UContextualAnimMovieSceneSection>(WeakSection.Get()))
+	{
+		return FText::FromString(GetNameSafe(Section->GetAnimTrack().Animation));
+		
+	}
+
+	return FText::FromString(TEXT("Animation_Name"));
+}
+
+FReply FContextualAnimSection::OnSectionDoubleClicked(const FGeometry& SectionGeometry, const FPointerEvent& MouseEvent) 
+{ 
+	if (UContextualAnimMovieSceneSection* Section = Cast<UContextualAnimMovieSceneSection>(WeakSection.Get()))
+	{
+		if(Section->IsActive())
+		{
+			//@TODO: Jump to Edit Notifies view
+		}
+		else
+		{
+			Section->GetOwnerTrack().GetViewModel().SetActiveAnimSetForSection(Section->GetSectionIdx(), Section->GetAnimSetIdx());
+		}
+
+		return FReply::Handled();
+	}
+	
+	return FReply::Unhandled();
+}
+
+float FContextualAnimSection::GetSectionHeight() const
+{
+	return 20.f;
+}
+
+#undef LOCTEXT_NAMESPACE
\ No newline at end of file
diff --git a/Engine/Plugins/Experimental/Animation/ContextualAnimation/Source/ContextualAnimationEditor/Private/ContextualAnimViewModel.cpp b/Engine/Plugins/Experimental/Animation/ContextualAnimation/Source/ContextualAnimationEditor/Private/ContextualAnimViewModel.cpp
--- a/Engine/Plugins/Experimental/Animation/ContextualAnimation/Source/ContextualAnimationEditor/Private/ContextualAnimViewModel.cpp
+++ b/Engine/Plugins/Experimental/Animation/ContextualAnimation/Source/ContextualAnimationEditor/Private/ContextualAnimViewModel.cpp
@@ -4,10 +4,10 @@
 #include "Modules/ModuleManager.h"
 #include "MovieScene.h"
 #include "ISequencerModule.h"
-#include "Tracks/MovieSceneSkeletalAnimationTrack.h"
-#include "Tracks/MovieSceneEventTrack.h"
-#include "Sections/MovieSceneEventRepeaterSection.h"
+#include "MovieSceneCommonHelpers.h"
 #include "ContextualAnimMovieSceneSequence.h"
+#include "ContextualAnimMovieSceneTrack.h"
+#include "ContextualAnimMovieSceneSection.h"
 #include "ContextualAnimMovieSceneNotifyTrackEditor.h"
 #include "ContextualAnimMovieSceneNotifyTrack.h"
 #include "ContextualAnimMovieSceneNotifySection.h"
@@ -74,11 +74,6 @@
 	RefreshSequencerTracks();
 }
 
-UAnimSequenceBase* FContextualAnimViewModel::FindAnimationByGuid(const FGuid& Guid) const
-{
-	return SceneInstance.IsValid() ? SceneInstance->FindBindingByGuid(Guid)->GetAnimTrack().Animation : nullptr;
-}
-
 void FContextualAnimViewModel::CreateSequencer()
 {
 	MovieSceneSequence = NewObject<UContextualAnimMovieSceneSequence>(GetTransientPackage());
@@ -110,11 +105,13 @@
 	Sequencer->SetPlaybackStatus(EMovieScenePlayerStatus::Stopped);
 }
 
-void FContextualAnimViewModel::SetActiveSceneVariantIdx(int32 Index)
+void FContextualAnimViewModel::SetActiveAnimSetForSection(int32 SectionIdx, int32 AnimSetIdx)
 {
-	check(FMath::IsWithinInclusive(Index, 0, (SceneAsset->GetTotalVariants() - 1)));
+	check(GetSceneAsset()->Sections.IsValidIndex(SectionIdx));
+	check(GetSceneAsset()->Sections[SectionIdx].AnimSets.IsValidIndex(AnimSetIdx));
 
-	ActiveSceneVariantIdx = Index;
+	int32& ActiveSetIdx = ActiveAnimSetMap.FindOrAdd(SectionIdx);
+	ActiveSetIdx = AnimSetIdx;
 
 	RefreshSequencerTracks();
 }
@@ -123,7 +120,7 @@
 {
 	const FContextualAnimRoleDefinition* RoleDef = GetSceneAsset()->RolesAsset ? GetSceneAsset()->RolesAsset->FindRoleDefinitionByName(AnimTrack.Role) : nullptr;
 	UClass* PreviewClass = RoleDef ? RoleDef->PreviewActorClass : nullptr;
-	const FTransform SpawnTransform = (AnimTrack.AlignmentData.ExtractTransformAtTime(0, 0.f));
+	const FTransform SpawnTransform = AnimTrack.GetRootTransformAtTime(0.f);
 
 	FActorSpawnParameters Params;
 	Params.SpawnCollisionHandlingOverride = ESpawnActorCollisionHandlingMethod::AdjustIfPossibleButAlwaysSpawn;
@@ -159,32 +156,32 @@
 {
 	// Remove movie scene tracks and destroy existing actors (if any)
 
-	for(UAnimSequenceBase* Anim : AnimsBeingEdited)
+	for (int32 MasterTrackIdx = MovieScene->GetMasterTracks().Num() - 1; MasterTrackIdx >= 0; MasterTrackIdx--)
 	{
-		Anim->UnregisterOnNotifyChanged(this);
+		UMovieSceneTrack& MasterTrack = *MovieScene->GetMasterTracks()[MasterTrackIdx];
+		for (const UMovieSceneSection* Section : MasterTrack.GetAllSections())
+		{
+			const UContextualAnimMovieSceneSection* ContextualAnimSection = Cast<const UContextualAnimMovieSceneSection>(Section);
+			check(ContextualAnimSection);
+
+			if(ContextualAnimSection->GetAnimTrack().Animation)
+			{
+				ContextualAnimSection->GetAnimTrack().Animation->UnregisterOnNotifyChanged(this);
+			}
+		}
+
+		MovieScene->RemoveMasterTrack(MasterTrack);
 	}
-	AnimsBeingEdited.Reset();
-
+			
 	if (SceneInstance.IsValid())
 	{
 		SceneInstance->Stop();
 	}
 
-	const int32 Num = MovieSceneSequence->GetMovieScene()->GetPossessableCount();
-	if(Num > 0)
-	{
-		for (int32 Idx = (Num - 1); Idx >= 0; Idx--)
-		{
-			FMovieScenePossessable& Possessable = MovieSceneSequence->GetMovieScene()->GetPossessable(Idx);
-			MovieSceneSequence->GetMovieScene()->RemovePossessable(Possessable.GetGuid());
-		}
-	}
-
 	if(StartSceneParams.RoleToActorMap.Num() > 0)
 	{
 		for (auto& MapEntry : StartSceneParams.RoleToActorMap)
-		{
-			
+		{			
 			if (AActor* Actor = MapEntry.Value.GetActor())
 			{
 				Actor->Destroy();
@@ -195,116 +192,170 @@
 	Sequencer->NotifyMovieSceneDataChanged(EMovieSceneDataChangeType::MovieSceneStructureItemsChanged);
 
 	StartSceneParams.Reset();
-	StartSceneParams.VariantIdx = ActiveSceneVariantIdx;
+	StartSceneParams.AnimSetIdx = 0;
 
-	SceneAsset->ForEachAnimTrack(ActiveSceneVariantIdx, [this](const FContextualAnimTrack& AnimTrack)
+	TArray<FName> Roles = SceneAsset->GetRoles();
+	for(FName Role : Roles)
+	{
+		UContextualAnimMovieSceneTrack* MovieSceneAnimTrack = MovieSceneSequence->GetMovieScene()->AddMasterTrack<UContextualAnimMovieSceneTrack>();
+		check(MovieSceneAnimTrack);
+
+		MovieSceneAnimTrack->Initialize(Role);
+	}
+
+	float SectionMaxTime = 0.f;
+	float SectionStartTime = 0.f;
+	for (int32 SectionIdx = 0; SectionIdx < SceneAsset->Sections.Num(); SectionIdx++)
+	{
+		SectionStartTime = SectionMaxTime;
+
+		FContextualAnimSceneSection& ContextualAnimSection = SceneAsset->Sections[SectionIdx];
+		for (int32 AnimSetIdx = 0; AnimSetIdx < ContextualAnimSection.AnimSets.Num(); AnimSetIdx++)
 		{
-			const FName& Role = AnimTrack.Role;
+			FContextualAnimSet& AnimSet = ContextualAnimSection.AnimSets[AnimSetIdx];
+			for (int32 AnimTrackIdx = 0; AnimTrackIdx < AnimSet.Tracks.Num(); AnimTrackIdx++)
+			{
+				const FContextualAnimTrack& AnimTrack = AnimSet.Tracks[AnimTrackIdx];
+				
+				UAnimSequenceBase* Animation = AnimTrack.Animation;
+				if (Animation)
+				{
+					UContextualAnimMovieSceneTrack* MovieSceneTrack = FindMasterTrackByRole(AnimTrack.Role);
+					if (MovieSceneTrack == nullptr)
+					{
+						UE_LOG(LogContextualAnim, Warning, TEXT("FContextualAnimViewModel::RefreshSequencerTracks. Can't find MovieSceneTrack for %s. Role: %s SectionIdx: %d AnimIndex: %d"), 
+							*GetNameSafe(Animation), *AnimTrack.Role.ToString(), SectionIdx, AnimSetIdx);
 
-			// Spawn preview actor
-			AActor* PreviewActor = SpawnPreviewActor(AnimTrack);
-			
-			if(PreviewActor == nullptr)
-			{
-				return UE::ContextualAnim::EForEachResult::Continue;
-			}
+						continue;
+					}
 
-			// Set actor label so the track shows the name of the role it represents
-			//PreviewActor->SetActorLabel(Role.ToString(), false);
-
-			// Add preview actors to sequencer
-			const bool bSelectActors = false;
-			TArray<TWeakObjectPtr<AActor>> Actors = { PreviewActor };
-			TArray<FGuid> Guids = Sequencer->AddActors(Actors, bSelectActors);
-			check(Guids.Num() > 0);
-
-			const FGuid& Guid = Guids[0];
-
-			UAnimSequenceBase* Animation = AnimTrack.Animation;
-			if (Animation)
-			{
-				// Add Animation Track
-				{
-					// @TODO: Temporally using an EventTrack to represent the animation since this is just a visual representation of the data. This assumes there is a single section in the montage
-					UMovieSceneEventTrack* MovieSceneAnimTrack = MovieSceneSequence->GetMovieScene()->AddTrack<UMovieSceneEventTrack>(Guid);
-					check(MovieSceneAnimTrack);
-
-					MovieSceneAnimTrack->SetDisplayName(FText::FromString(GetNameSafe(Animation)));
-
-					UMovieSceneSection* NewSection = NewObject<UMovieSceneSection>(MovieSceneAnimTrack, UMovieSceneEventRepeaterSection::StaticClass(), NAME_None, RF_Transactional);
+					UContextualAnimMovieSceneSection* NewSection = NewObject<UContextualAnimMovieSceneSection>(MovieSceneTrack, UContextualAnimMovieSceneSection::StaticClass(), NAME_None, RF_Transactional);
 					check(NewSection);
 
-					FFrameRate TickResolution = MovieSceneSequence->GetMovieScene()->GetTickResolution();
-					FFrameNumber StartFrame(0);
-					FFrameNumber EndFrame = (Animation->GetPlayLength() * TickResolution).RoundToFrame();
-					NewSection->SetRange(TRange<FFrameNumber>::Exclusive(StartFrame, EndFrame));
+					NewSection->Initialize(SectionIdx, AnimSetIdx, AnimTrackIdx);
 
-					MovieSceneAnimTrack->AddSection(*NewSection);
+					const float AnimLength = Animation->GetPlayLength();
+					const FFrameRate TickResolution = MovieSceneSequence->GetMovieScene()->GetTickResolution();
+					const FFrameNumber StartFrame = (SectionStartTime * TickResolution).RoundToFrame();
+					const FFrameNumber EndFrame = (AnimLength * TickResolution).RoundToFrame() + StartFrame;
+					NewSection->SetRange(TRange<FFrameNumber>::Inclusive(StartFrame, EndFrame));
+
+					NewSection->SetRowIndex(AnimSetIdx);
+
+					const int32& ActiveSetIdx = ActiveAnimSetMap.FindOrAdd(SectionIdx);
+					NewSection->SetIsActive(AnimSetIdx == ActiveSetIdx);
+
+					MovieSceneTrack->AddSection(*NewSection);
+					MovieSceneTrack->SetTrackRowDisplayName(FText::FromString(FString::Printf(TEXT("%d"), AnimSetIdx)), AnimSetIdx);
+
+					if (AnimLength > SectionMaxTime)
+					{
+						SectionMaxTime = AnimLength;
+					}
 				}
 
-				// Add Notify Tracks
+				if (!StartSceneParams.RoleToActorMap.Contains(AnimTrack.Role))
 				{
-					for (const FAnimNotifyTrack& NotifyTrack : Animation->AnimNotifyTracks)
+					if (AActor* PreviewActor = SpawnPreviewActor(AnimTrack))
 					{
-						UContextualAnimMovieSceneNotifyTrack* Track = MovieSceneSequence->GetMovieScene()->AddTrack<UContextualAnimMovieSceneNotifyTrack>(Guid);
-						check(Track);
-
-						Track->Initialize(*Animation, NotifyTrack);
+						StartSceneParams.RoleToActorMap.Add(AnimTrack.Role, PreviewActor);
 					}
-
-					// Listen for when the notifies in the animation changes, so we can refresh the notify sections here
-					Animation->RegisterOnNotifyChanged(UAnimSequenceBase::FOnNotifyChanged::CreateSP(this, &FContextualAnimViewModel::OnAnimNotifyChanged, Animation));
 				}
 
-				AnimsBeingEdited.Add(Animation);
-			}
-
-			StartSceneParams.RoleToActorMap.Add(Role, PreviewActor);
-
-			return UE::ContextualAnim::EForEachResult::Continue;
-		});
+			} // End AnimTracks Loop
+		} // End AnimSets Loop
+	} // End Sections Loop
 
 	Sequencer->NotifyMovieSceneDataChanged(EMovieSceneDataChangeType::MovieSceneStructureItemsChanged);
 
 	SceneInstance = ContextualAnimManager->ForceStartScene(*GetSceneAsset(), StartSceneParams);
 
 	// Disable auto blend out
-	if(ensureAlways(SceneInstance.IsValid()))
+	if(SceneInstance.IsValid())
 	{
 		for (auto& Binding : SceneInstance->GetBindings())
 		{
-			Binding.Guid = Sequencer->FindObjectId(*Binding.GetActor(), MovieSceneSequenceID::Root);
-			check(Binding.Guid.IsValid());
-
 			if (FAnimMontageInstance* MontageInstance = Binding.GetAnimMontageInstance())
 			{
+				MontageInstance->Pause();
 				MontageInstance->bEnableAutoBlendOut = false;
 			}
 		}
 	}
 }
 
-void FContextualAnimViewModel::AddNewVariant(const FContextualAnimNewVariantParams& Params)
+void FContextualAnimViewModel::AddNewAnimSet(const FContextualAnimNewAnimSetParams& Params)
 {
-	FContextualAnimTracksContainer Container;
-	for (const FContextualAnimNewVariantData& Data : Params.Data)
+	FContextualAnimSet AnimSet;
+	for (const FContextualAnimNewAnimSetData& Data : Params.Data)
 	{
 		FContextualAnimTrack AnimTrack;
 		AnimTrack.Role = Data.RoleName;
 		AnimTrack.Animation = Data.Animation;
 		AnimTrack.bRequireFlyingMode = Data.bRequiresFlyingMode;
-		Container.Tracks.Add(AnimTrack);
+		AnimSet.Tracks.Add(AnimTrack);
 	}
 
-	SceneAsset->Variants.Add(Container);
+	int32 SectionIdx = INDEX_NONE;
+	int32 AnimSetIdx = INDEX_NONE;
+	for(int32 Idx = 0; Idx < SceneAsset->Sections.Num(); Idx++)
+	{
+		FContextualAnimSceneSection& Section = SceneAsset->Sections[Idx];
+		if(Section.Name == Params.SectionName)
+		{
+			AnimSetIdx = Section.AnimSets.Add(AnimSet);
+			SectionIdx = Idx;
+			break;
+		}
+	}
+	
+	if(SectionIdx == INDEX_NONE)
+	{
+		FContextualAnimSceneSection NewSection;
+		NewSection.Name = Params.SectionName;
+		AnimSetIdx = NewSection.AnimSets.Add(AnimSet);
+		SectionIdx = SceneAsset->Sections.Add(NewSection);
+	}
+
+	check(SectionIdx != INDEX_NONE);
+	check(AnimSetIdx != INDEX_NONE);
 
 	SceneAsset->PrecomputeData();
 
 	SceneAsset->MarkPackageDirty();
 
-	SetActiveSceneVariantIdx(SceneAsset->GetTotalVariants() - 1);
+	// Set active AnimSet and refresh sequencer panel
+	SetActiveAnimSetForSection(SectionIdx, AnimSetIdx);
 }
 
+void FContextualAnimViewModel::AddNewIKTarget(const UContextualAnimNewIKTargetParams& Params)
+{
+	check(SceneAsset->Sections.IsValidIndex(Params.SectionIdx));
+
+	// Add IK Target definition to the scene asset
+	FContextualAnimIKTargetDefinition IKTargetDef;
+	IKTargetDef.GoalName = Params.GoalName;
+	IKTargetDef.BoneName = Params.SourceBone.BoneName;
+	IKTargetDef.Provider = Params.Provider;
+	IKTargetDef.TargetRoleName = Params.TargetRole;
+	IKTargetDef.TargetBoneName = Params.TargetBone.BoneName;
+
+	if (FContextualAnimIKTargetDefContainer* ContainerPtr = SceneAsset->Sections[Params.SectionIdx].RoleToIKTargetDefsMap.Find(Params.SourceRole))
+	{
+		ContainerPtr->IKTargetDefs.AddUnique(IKTargetDef);
+	}
+	else
+	{
+		FContextualAnimIKTargetDefContainer Container;
+		Container.IKTargetDefs.AddUnique(IKTargetDef);
+		SceneAsset->Sections[Params.SectionIdx].RoleToIKTargetDefsMap.Add(Params.SourceRole, Container);
+	}
+
+	SceneAsset->PrecomputeData();
+
+	SceneAsset->MarkPackageDirty();
+}
+
 void FContextualAnimViewModel::ToggleSimulateMode() 
 { 
 	bIsSimulateModeActive = !bIsSimulateModeActive; 
@@ -324,9 +375,9 @@
 			{
 				if (UMotionWarpingComponent* MotionWarpComp = Binding.GetActor()->FindComponentByClass<UMotionWarpingComponent>())
 				{
-					for (const FContextualAnimAlignmentSectionData AligmentData : SceneAsset->AlignmentSections)
+					for (const FContextualAnimSetPivotDefinition& Def : SceneAsset->GetAnimSetPivotDefinitionsInSection(Binding.GetAnimTrack().SectionIdx))
 					{
-						MotionWarpComp->RemoveWarpTarget(AligmentData.WarpTargetName);
+						MotionWarpComp->RemoveWarpTarget(Def.Name);
 					}
 				}
 			}
@@ -361,23 +412,50 @@
 	return GetWorld();
 }
 
+void FContextualAnimViewModel::UpdatePreviewActorTransform(const FContextualAnimSceneBinding& Binding, float Time)
+{
+	if (AActor* PreviewActor = Binding.GetActor())
+	{
+		FTransform Transform = Binding.GetAnimTrack().GetRootTransformAtTime(Time);
+
+		// Special case for Character
+		if (ACharacter* PreviewCharacter = Cast<ACharacter>(PreviewActor))
+		{
+			if (UCharacterMovementComponent* MovementComp = Binding.GetActor()->FindComponentByClass<UCharacterMovementComponent>())
+			{
+				MovementComp->StopMovementImmediately();
+			}
+
+			const float MIN_FLOOR_DIST = 1.9f; //from CharacterMovementComp, including in this offset to avoid jittering in walking mode
+			const float CapsuleHalfHeight = PreviewCharacter->GetCapsuleComponent()->GetScaledCapsuleHalfHeight();
+			Transform.SetLocation(Transform.GetLocation() + (PreviewCharacter->GetActorQuat().GetUpVector() * CapsuleHalfHeight + MIN_FLOOR_DIST));
+
+			Transform.SetRotation(PreviewCharacter->GetBaseRotationOffset().Inverse() * Transform.GetRotation());
+		}
+
+		PreviewActor->SetActorLocationAndRotation(Transform.GetLocation(), Transform.GetRotation());
+	}
+}
+
+UContextualAnimMovieSceneTrack* FContextualAnimViewModel::FindMasterTrackByRole(const FName& Role) const
+{
+	const TArray<UMovieSceneTrack*>& MasterTracks = MovieScene->GetMasterTracks();
+	for (UMovieSceneTrack* MasterTrack : MasterTracks)
+	{
+		UContextualAnimMovieSceneTrack* ContextualAnimTrack = Cast<UContextualAnimMovieSceneTrack>(MasterTrack);
+		check(ContextualAnimTrack);
+
+		if (ContextualAnimTrack->GetRole() == Role)
+		{
+			return ContextualAnimTrack;
+		}
+	}
+
+	return nullptr;
+}
+
 void FContextualAnimViewModel::SequencerTimeChanged()
 {
-	auto ResetActorTransform = [](FContextualAnimSceneBinding& Binding, float Time)
-	{
-		const USkeletalMeshComponent* SkelMeshComp = UContextualAnimUtilities::TryGetSkeletalMeshComponent(Binding.GetActor());
-
-		const FTransform RootTransform = UContextualAnimUtilities::ExtractRootTransformFromAnimation(Binding.GetAnimTrack().Animation, Time);
-		const FTransform StartTransform = SkelMeshComp->GetRelativeTransform().Inverse() * RootTransform;
-
-		Binding.GetActor()->SetActorLocationAndRotation(StartTransform.GetLocation(), StartTransform.GetRotation());
-
-		if (UCharacterMovementComponent* MovementComp = Binding.GetActor()->FindComponentByClass<UCharacterMovementComponent>())
-		{
-			MovementComp->StopMovementImmediately();
-		}
-	};
-
 	EMovieScenePlayerStatus::Type CurrentStatus = Sequencer->GetPlaybackStatus();
 	const float CurrentSequencerTime = Sequencer->GetGlobalTime().AsSeconds();
 	const float PlaybackSpeed = Sequencer->GetPlaybackSpeed();
@@ -386,35 +464,72 @@
 	{
 		for (auto& Binding : SceneInstance->GetBindings())
 		{
-			if (FAnimMontageInstance* MontageInstance = Binding.GetAnimMontageInstance())
+			FAnimMontageInstance* MontageInstance = Binding.GetAnimMontageInstance();
+			if (MontageInstance == nullptr)
 			{
-				const float AnimPlayLength = MontageInstance->Montage->GetPlayLength();
-				float PreviousTime = FMath::Clamp(PreviousSequencerTime, 0.f, AnimPlayLength);
-				float CurrentTime = FMath::Clamp(CurrentSequencerTime, 0.f, AnimPlayLength);
+				continue;
+			}
 
-				if (CurrentStatus == EMovieScenePlayerStatus::Stopped || CurrentStatus == EMovieScenePlayerStatus::Scrubbing)
+			const UContextualAnimMovieSceneTrack* MasterTrack = FindMasterTrackByRole(Binding.GetRoleDef().Name);
+			check(MasterTrack);
+
+			const FFrameNumber CurrentFrame = (CurrentSequencerTime * MovieSceneSequence->GetMovieScene()->GetTickResolution()).RoundToFrame();
+			const UContextualAnimMovieSceneSection* Section = Cast<const UContextualAnimMovieSceneSection>(MovieSceneHelpers::FindSectionAtTime(MasterTrack->GetAllSections(), CurrentFrame));
+
+			if (Section == nullptr || Section->GetAnimTrack().Animation == nullptr)
+			{
+				continue;
+			}
+
+			const UAnimSequenceBase* FocusedAnim = Section->GetAnimTrack().Animation;
+			const float AnimPlayLength = FocusedAnim->GetPlayLength();
+			const float SectionStartTime = (Section->GetInclusiveStartFrame() / MovieSceneSequence->GetMovieScene()->GetTickResolution());
+			const float PrevLocalTime = FMath::Clamp(PreviousSequencerTime - SectionStartTime, 0.f, AnimPlayLength);
+			const float CurrLocalTime = FMath::Clamp(CurrentSequencerTime - SectionStartTime, 0.f, AnimPlayLength);
+
+			bool bShouldTransition = false;
+			if(const UAnimMontage* FocusedAnimAsMontage = Cast<const UAnimMontage>(FocusedAnim))
+			{
+				bShouldTransition = FocusedAnimAsMontage != MontageInstance->Montage;
+			}
+			else if(FocusedAnim)
+			{
+				bShouldTransition = FocusedAnim != MontageInstance->Montage->SlotAnimTracks[0].AnimTrack.AnimSegments[0].GetAnimReference();
+			}
+
+			if (bShouldTransition)
+			{
+				SceneInstance->TransitionTo(Binding, Section->GetAnimTrack());
+				MontageInstance = Binding.GetAnimMontageInstance();
+				check(MontageInstance);
+
+				MontageInstance->bEnableAutoBlendOut = false;
+				MontageInstance->SetPosition(CurrLocalTime);
+				UpdatePreviewActorTransform(Binding, CurrLocalTime);
+			}
+
+			if (CurrentStatus == EMovieScenePlayerStatus::Stopped || CurrentStatus == EMovieScenePlayerStatus::Scrubbing)
+			{
+				UpdatePreviewActorTransform(Binding, CurrLocalTime);
+
+				if (MontageInstance->IsPlaying())
 				{
-					ResetActorTransform(Binding, CurrentTime);
+					MontageInstance->Pause();
+				}
 
-					if (MontageInstance->IsPlaying())
-					{
-						MontageInstance->Pause();
-					}
+				MontageInstance->SetPosition(CurrLocalTime);
+			}
+			else if (CurrentStatus == EMovieScenePlayerStatus::Playing)
+			{
+				if (PlaybackSpeed > 0.f && CurrLocalTime < PrevLocalTime)
+				{
+					UpdatePreviewActorTransform(Binding, CurrLocalTime);
+					MontageInstance->SetPosition(CurrLocalTime);
+				}
 
-					MontageInstance->SetPosition(CurrentTime);
-				}
-				else if (CurrentStatus == EMovieScenePlayerStatus::Playing)
+				if (!MontageInstance->IsPlaying())
 				{
-					if (PlaybackSpeed > 0.f && CurrentTime < PreviousTime)
-					{
-						ResetActorTransform(Binding, CurrentTime);
-						MontageInstance->SetPosition(CurrentTime);
-					}
-
-					if (!MontageInstance->IsPlaying())
-					{
-						MontageInstance->SetPlaying(true);
-					}
+					MontageInstance->SetPlaying(true);
 				}
 			}
 		}
@@ -438,7 +553,8 @@
 		// Update IK AnimNotify's bEnable flag based on the Active state of the section
 		// @TODO: Temp brute-force approach until having a way to override FMovieSceneSection::SetIsActive or something similar
 
-		for (const auto& Binding : SceneInstance->GetBindings())
+		// @TODO: Commented out for now until we add the new behavior where the user needs to double-click on the animation to edit the notifies
+		/*for (const auto& Binding : SceneInstance->GetBindings())
 		{
 			TArray<UMovieSceneTrack*> Tracks = MovieSceneSequence->GetMovieScene()->FindTracks(UContextualAnimMovieSceneNotifyTrack::StaticClass(), Binding.Guid);
 			for(UMovieSceneTrack* Track : Tracks)
@@ -459,7 +575,7 @@
 					}
 				}
 			}
-		}
+		}*/
 	}
 }
 
diff --git a/Engine/Plugins/Experimental/Animation/ContextualAnimation/Source/ContextualAnimationEditor/Private/ContextualAnimationEditor.cpp b/Engine/Plugins/Experimental/Animation/ContextualAnimation/Source/ContextualAnimationEditor/Private/ContextualAnimationEditor.cpp
--- a/Engine/Plugins/Experimental/Animation/ContextualAnimation/Source/ContextualAnimationEditor/Private/ContextualAnimationEditor.cpp
+++ b/Engine/Plugins/Experimental/Animation/ContextualAnimation/Source/ContextualAnimationEditor/Private/ContextualAnimationEditor.cpp
@@ -4,6 +4,7 @@
 #include "ContextualAnimEdMode.h"
 #include "ContextualAnimEditorStyle.h"
 #include "ContextualAnimTypeActions.h"
+#include "ContextualAnimMovieSceneTrackEditor.h"
 #include "ContextualAnimMovieSceneNotifyTrackEditor.h"
 #include "ContextualAnimAssetEditorCommands.h"
 #include "DetailCustomizations/ContextualAnimNotifySectionDetailCustom.h"
@@ -35,6 +36,7 @@
 	// Register MovieSceneAnimNotifyTrackEditor
 	ISequencerModule& SequencerModule = FModuleManager::LoadModuleChecked< ISequencerModule >("Sequencer");
 	MovieSceneAnimNotifyTrackEditorHandle = SequencerModule.RegisterTrackEditor(FOnCreateTrackEditor::CreateStatic(&FContextualAnimMovieSceneNotifyTrackEditor::CreateTrackEditor));
+	MovieSceneAnimNotifyTrackEditorHandle = SequencerModule.RegisterTrackEditor(FOnCreateTrackEditor::CreateStatic(&FContextualAnimMovieSceneTrackEditor::CreateTrackEditor));
 
 	// Register Detail Customizations
 	FPropertyEditorModule& PropertyModule = FModuleManager::LoadModuleChecked<FPropertyEditorModule>("PropertyEditor");;
diff --git a/Engine/Plugins/Experimental/Animation/ContextualAnimation/Source/ContextualAnimationEditor/Public/ContextualAnimAssetEditorToolkit.h b/Engine/Plugins/Experimental/Animation/ContextualAnimation/Source/ContextualAnimationEditor/Public/ContextualAnimAssetEditorToolkit.h
--- a/Engine/Plugins/Experimental/Animation/ContextualAnimation/Source/ContextualAnimationEditor/Public/ContextualAnimAssetEditorToolkit.h
+++ b/Engine/Plugins/Experimental/Animation/ContextualAnimation/Source/ContextualAnimationEditor/Public/ContextualAnimAssetEditorToolkit.h
@@ -52,7 +52,7 @@
 	void ExtendToolbar();
 	void FillToolbar(FToolBarBuilder& ToolbarBuilder);
 
-	TSharedRef<SWidget> BuildVariantsMenu();
+	TSharedRef<SWidget> BuildNewAnimSetWidget();
 
 	void OnFinishedChangingProperties(const FPropertyChangedEvent& PropertyChangedEvent);
 
@@ -66,7 +66,7 @@
 
 	TSharedPtr<FContextualAnimViewModel> ViewModel;
 
-	TSharedPtr<FStructOnScope> NewVariantWidgetStruct;
+	TSharedPtr<FStructOnScope> NewAnimSetWidgetStruct;
 };
 
 
diff --git a/Engine/Plugins/Experimental/Animation/ContextualAnimation/Source/ContextualAnimationEditor/Public/ContextualAnimEdMode.h b/Engine/Plugins/Experimental/Animation/ContextualAnimation/Source/ContextualAnimationEditor/Public/ContextualAnimEdMode.h
--- a/Engine/Plugins/Experimental/Animation/ContextualAnimation/Source/ContextualAnimationEditor/Public/ContextualAnimEdMode.h
+++ b/Engine/Plugins/Experimental/Animation/ContextualAnimation/Source/ContextualAnimationEditor/Public/ContextualAnimEdMode.h
@@ -10,15 +10,17 @@
 
 struct FSelectionCriterionHitProxyData
 {
+	int32 SectionIdx = INDEX_NONE;
+	int32 AnimSetIdx = INDEX_NONE;
 	FName RoleName = NAME_None;
-	int32 VariantIdx = INDEX_NONE;
 	int32 CriterionIdx = INDEX_NONE;
 	int32 DataIdx = INDEX_NONE;
 
 	FSelectionCriterionHitProxyData(){}
-	FSelectionCriterionHitProxyData(const FName& InRoleName, int32 InVariantIdx, int32 InCriterionIdx, int32 InDataIdx)
-		: RoleName(InRoleName)
-		, VariantIdx(InVariantIdx)
+	FSelectionCriterionHitProxyData(int32 InSectionIdx, int32 InAnimSetIdx, const FName& InRoleName, int32 InCriterionIdx, int32 InDataIdx)
+		: SectionIdx(InSectionIdx)
+		, AnimSetIdx(InAnimSetIdx)
+		, RoleName(InRoleName)
 		, CriterionIdx(InCriterionIdx)
 		, DataIdx(InDataIdx)
 	{}
@@ -26,7 +28,8 @@
 	void Reset()
 	{
 		RoleName = NAME_None;
-		VariantIdx = INDEX_NONE;
+		SectionIdx = INDEX_NONE;
+		AnimSetIdx = INDEX_NONE;
 		CriterionIdx = INDEX_NONE;
 		DataIdx = INDEX_NONE;
 	}
@@ -34,7 +37,8 @@
 	bool IsValid() const 
 	{ 
 		return	RoleName != NAME_None && 
-				VariantIdx != INDEX_NONE && 
+				SectionIdx != INDEX_NONE &&
+				AnimSetIdx != INDEX_NONE && 
 				CriterionIdx != INDEX_NONE; 
 	}
 };
diff --git a/Engine/Plugins/Experimental/Animation/ContextualAnimation/Source/ContextualAnimationEditor/Public/ContextualAnimEditorTypes.h b/Engine/Plugins/Experimental/Animation/ContextualAnimation/Source/ContextualAnimationEditor/Public/ContextualAnimEditorTypes.h
--- a/Engine/Plugins/Experimental/Animation/ContextualAnimation/Source/ContextualAnimationEditor/Public/ContextualAnimEditorTypes.h
+++ b/Engine/Plugins/Experimental/Animation/ContextualAnimation/Source/ContextualAnimationEditor/Public/ContextualAnimEditorTypes.h
@@ -35,16 +35,20 @@
 	UPROPERTY(EditAnywhere, Category = "Settings")
 	FName GoalName = NAME_None;
 
+	UPROPERTY()
+	int32 SectionIdx = INDEX_NONE;
+
 	UContextualAnimNewIKTargetParams(const FObjectInitializer& ObjectInitializer) : Super(ObjectInitializer) {}
 
 	// ~IBoneReferenceSkeletonProvider Interface
 	virtual USkeleton* GetSkeleton(bool& bInvalidSkeletonIsError, const IPropertyHandle* PropertyHandle = nullptr) override;
 
-	void Reset(const FName& InSourceRole, const UContextualAnimSceneAsset& InSceneAsset);
+	void Reset(const UContextualAnimSceneAsset& InSceneAsset, const UAnimSequenceBase& InAnimation);
 
 	bool HasValidData() const;
 
 	const UContextualAnimSceneAsset& GetSceneAsset() const;
+	int32 GetSectionIdx() const { return SectionIdx; }
 
 	UFUNCTION()
 	TArray<FString> GetTargetRoleOptions() const;
@@ -59,7 +63,7 @@
 };
 
 USTRUCT()
-struct FContextualAnimNewVariantData
+struct FContextualAnimNewAnimSetData
 {
 	GENERATED_BODY()
 
@@ -73,12 +77,15 @@
 	bool bRequiresFlyingMode = false;
 };
 
-/** Struct used to construct the widget for adding a new variant */
+/** Struct used to construct the widget for adding a new set */
 USTRUCT()
-struct FContextualAnimNewVariantParams
+struct FContextualAnimNewAnimSetParams
 {
 	GENERATED_BODY()
 
+	UPROPERTY(EditAnywhere, Category = "Settings")
+	FName SectionName = NAME_Default;
+
 	UPROPERTY(EditAnywhere, EditFixedSize, Category = "Settings", meta = (TitleProperty = "RoleName"))
-	TArray<FContextualAnimNewVariantData> Data;
+	TArray<FContextualAnimNewAnimSetData> Data;
 };
diff --git a/Engine/Plugins/Experimental/Animation/ContextualAnimation/Source/ContextualAnimationEditor/Public/ContextualAnimMovieSceneSection.h b/Engine/Plugins/Experimental/Animation/ContextualAnimation/Source/ContextualAnimationEditor/Public/ContextualAnimMovieSceneSection.h
new file mode 100644
--- /dev/null
+++ b/Engine/Plugins/Experimental/Animation/ContextualAnimation/Source/ContextualAnimationEditor/Public/ContextualAnimMovieSceneSection.h
@@ -0,0 +1,40 @@
+// Copyright Epic Games, Inc. All Rights Reserved.
+
+#pragma once
+
+#include "CoreMinimal.h"
+#include "MovieSceneSection.h"
+#include "ContextualAnimMovieSceneSection.generated.h"
+
+struct FContextualAnimTrack;
+class UContextualAnimMovieSceneTrack;
+
+UCLASS()
+class UContextualAnimMovieSceneSection : public UMovieSceneSection
+{
+	GENERATED_BODY()
+
+public:
+
+	void Initialize(int32 InSectionIdx, int32 InAnimSetIdx, int32 InAnimTrackIdx);
+
+	/** Returns the track owner of this section */
+	UContextualAnimMovieSceneTrack& GetOwnerTrack() const;
+
+	FContextualAnimTrack& GetAnimTrack() const;
+
+	FORCEINLINE int32 GetSectionIdx() const { return SectionIdx; }
+	FORCEINLINE int32 GetAnimSetIdx() const { return AnimSetIdx; }
+	FORCEINLINE int32 GetAnimTrackIdx() const { return AnimTrackIdx; }
+
+private:
+
+	UPROPERTY()
+	int32 SectionIdx = INDEX_NONE;
+
+	UPROPERTY()
+	int32 AnimSetIdx = INDEX_NONE;
+
+	UPROPERTY()
+	int32 AnimTrackIdx = INDEX_NONE;
+};
\ No newline at end of file
diff --git a/Engine/Plugins/Experimental/Animation/ContextualAnimation/Source/ContextualAnimationEditor/Public/ContextualAnimMovieSceneTrack.h b/Engine/Plugins/Experimental/Animation/ContextualAnimation/Source/ContextualAnimationEditor/Public/ContextualAnimMovieSceneTrack.h
new file mode 100644
--- /dev/null
+++ b/Engine/Plugins/Experimental/Animation/ContextualAnimation/Source/ContextualAnimationEditor/Public/ContextualAnimMovieSceneTrack.h
@@ -0,0 +1,66 @@
+// Copyright Epic Games, Inc. All Rights Reserved.
+
+#pragma once
+
+#include "CoreMinimal.h"
+#include "MovieSceneNameableTrack.h"
+#include "ContextualAnimMovieSceneTrack.generated.h"
+
+class FContextualAnimViewModel;
+
+UCLASS()
+class UContextualAnimMovieSceneTrack : public UMovieSceneNameableTrack
+{
+	GENERATED_BODY()
+
+public:
+	
+	UContextualAnimMovieSceneTrack();
+
+	void Initialize(const FName& InRole);
+
+	FORCEINLINE FName GetRole() const { return Role; }
+
+	FContextualAnimViewModel& GetViewModel() const;
+
+	/** Creates a new section for this track */
+	virtual UMovieSceneSection* CreateNewSection() override;
+	
+	/** Adds a new section to this track */
+	virtual void AddSection(UMovieSceneSection& Section) override;
+	
+	/** Returns whether this track supports the supplied section class */
+	virtual bool SupportsType(TSubclassOf<UMovieSceneSection> SectionClass) const override;
+	
+	/** Returns the list of all the sections in this track */
+	virtual const TArray<UMovieSceneSection*>& GetAllSections() const override;
+	
+	/** Returns whether this track contains the supplied section */
+	virtual bool HasSection(const UMovieSceneSection& Section) const override;
+
+	/** Returns whether this track is empty (doesn't have any section in it) */
+	virtual bool IsEmpty() const override;
+	
+	/** Called when the supplied section is removed from Sequencer. Removes the actual Notify the section is representing */
+	virtual void RemoveSection(UMovieSceneSection& Section) override;
+
+	/** Called when a section at the supplied index is removed from Sequencer. Removes the actual Notify the section is representing */
+	virtual void RemoveSectionAt(int32 SectionIndex) override;
+
+#if WITH_EDITOR
+	/** Called when a section in this track is moved. Updates the actual Notify the section is representing */
+	virtual EMovieSceneSectionMovedResult OnSectionMoved(UMovieSceneSection& Section, const FMovieSceneSectionMovedParams& Params) override;
+#endif
+
+private:
+
+	/** List of sections in this track */
+	UPROPERTY()
+	TArray<TObjectPtr<UMovieSceneSection>> Sections;
+
+	/** Role this track represents */
+	UPROPERTY()
+	FName Role = NAME_None;
+
+};
+
diff --git a/Engine/Plugins/Experimental/Animation/ContextualAnimation/Source/ContextualAnimationEditor/Public/ContextualAnimMovieSceneTrackEditor.h b/Engine/Plugins/Experimental/Animation/ContextualAnimation/Source/ContextualAnimationEditor/Public/ContextualAnimMovieSceneTrackEditor.h
new file mode 100644
--- /dev/null
+++ b/Engine/Plugins/Experimental/Animation/ContextualAnimation/Source/ContextualAnimationEditor/Public/ContextualAnimMovieSceneTrackEditor.h
@@ -0,0 +1,42 @@
+// Copyright Epic Games, Inc. All Rights Reserved.
+
+#pragma once
+
+#include "CoreMinimal.h"
+#include "MovieSceneTrackEditor.h"
+#include "ContextualAnimTypes.h"
+
+/** Handles section drawing and manipulation of a ContextualAnimMovieSceneTrack */
+class FContextualAnimMovieSceneTrackEditor : public FMovieSceneTrackEditor
+{
+public:
+
+	FContextualAnimMovieSceneTrackEditor(TSharedRef<ISequencer> InSequencer);
+
+	// ~FMovieSceneTrackEditor Interface
+	virtual TSharedRef<ISequencerSection> MakeSectionInterface(UMovieSceneSection& SectionObject, UMovieSceneTrack& Track, FGuid ObjectBinding) override;
+	virtual bool SupportsType(TSubclassOf<UMovieSceneTrack> TrackClass) const override;
+	virtual bool SupportsSequence(UMovieSceneSequence* InSequence) const override;
+	virtual TSharedPtr<SWidget> BuildOutlinerEditWidget(const FGuid& ObjectBinding, UMovieSceneTrack* Track, const FBuildEditWidgetParams& Params) override;
+	virtual void BuildAddTrackMenu(FMenuBuilder& MenuBuilder) override;
+	virtual void BuildObjectBindingTrackMenu(FMenuBuilder& MenuBuilder, const TArray<FGuid>& ObjectBindings, const UClass* ObjectClass) override;
+
+	static TSharedRef<ISequencerTrackEditor> CreateTrackEditor(TSharedRef<ISequencer> InSequencer);
+};
+
+// FContextualAnimSection
+////////////////////////////////////////////////////////////////////////////////////////////////
+
+/** UI portion of a NotifySection in a NotifyTrack */
+class FContextualAnimSection : public FSequencerSection
+{
+public:
+	
+	FContextualAnimSection(UMovieSceneSection& InSectionObject)
+		: FSequencerSection(InSectionObject)
+	{}
+
+	virtual FReply OnSectionDoubleClicked(const FGeometry& SectionGeometry, const FPointerEvent& MouseEvent) override;
+	virtual FText GetSectionTitle() const override;
+	virtual float GetSectionHeight() const override;
+};
\ No newline at end of file
diff --git a/Engine/Plugins/Experimental/Animation/ContextualAnimation/Source/ContextualAnimationEditor/Public/ContextualAnimViewModel.h b/Engine/Plugins/Experimental/Animation/ContextualAnimation/Source/ContextualAnimationEditor/Public/ContextualAnimViewModel.h
--- a/Engine/Plugins/Experimental/Animation/ContextualAnimation/Source/ContextualAnimationEditor/Public/ContextualAnimViewModel.h
+++ b/Engine/Plugins/Experimental/Animation/ContextualAnimation/Source/ContextualAnimationEditor/Public/ContextualAnimViewModel.h
@@ -18,13 +18,15 @@
 class UMovieScene;
 class UMovieSceneTrack;
 class UMovieSceneSection;
+class UContextualAnimMovieSceneTrack;
 class UContextualAnimMovieSceneNotifyTrack;
 class UContextualAnimMovieSceneNotifySection;
 class UAnimSequenceBase;
 class IDetailsView;
 class IStructureDetailsView;
+class UContextualAnimNewIKTargetParams;
 struct FMovieSceneSectionMovedParams;
-struct FContextualAnimNewVariantParams;
+struct FContextualAnimNewAnimSetParams;
 struct FContextualAnimTrack;
 
 class FContextualAnimViewModel : public TSharedFromThis<FContextualAnimViewModel>, public FGCObject
@@ -48,23 +50,24 @@
 	UContextualAnimMovieSceneSequence* GetMovieSceneSequence() const { return MovieSceneSequence; }
 	UContextualAnimSceneAsset* GetSceneAsset() const { return SceneAsset; }
 	UContextualAnimSceneInstance* GetSceneInstance() const { return SceneInstance.Get(); }
+	UContextualAnimMovieSceneTrack* FindMasterTrackByRole(const FName& Role) const;
 
-	void AddNewVariant(const FContextualAnimNewVariantParams& Params);
+	void AddNewAnimSet(const FContextualAnimNewAnimSetParams& Params);
 
-	UAnimSequenceBase* FindAnimationByGuid(const FGuid& Guid) const;
+	void AddNewIKTarget(const UContextualAnimNewIKTargetParams& Params);
 
 	void AnimationModified(UAnimSequenceBase& Animation);
 
-	void SetActiveSceneVariantIdx(int32 Index);
+	void SetActiveAnimSetForSection(int32 SectionIdx, int32 AnimSetIdx);
 
-	int32 GetActiveSceneVariantIdx() const { return ActiveSceneVariantIdx; }
-
 	void OnPreviewActorClassChanged();
 
 	void ToggleSimulateMode();
 	bool IsSimulateModeActive() { return bIsSimulateModeActive; }
 	void StartSimulation();
 
+	void UpdatePreviewActorTransform(const FContextualAnimSceneBinding& Binding, float Time);
+
 private:
 
 	/** Scene asset being viewed and edited by this view model. */
@@ -86,7 +89,8 @@
 
 	TWeakObjectPtr<UContextualAnimSceneInstance> SceneInstance;
 
-	int32 ActiveSceneVariantIdx = 0;
+	/** Active anim set for each section. Key = SectionIdx, Value = AnimSetIdx  */
+	TMap<int32, int32> ActiveAnimSetMap;
 
 	/** The previous play status for sequencer time line. */
 	EMovieScenePlayerStatus::Type PreviousSequencerStatus;
@@ -101,9 +105,6 @@
 
 	FContextualAnimStartSceneParams StartSceneParams;
 
-	/** Container for the animations on the time line. Should be removed once we add a proper animation track */
-	TArray<UAnimSequenceBase*> AnimsBeingEdited;
-
 	AActor* SpawnPreviewActor(const FContextualAnimTrack& AnimTrack);
 
 	UWorld* GetWorld() const;
