[ContextualAnimation] misc changes for runtime playback
- added new multicast delegate when all montages are done playing (similar to OnSceneEnded but can be used by non-UObject systems)
- added new multicast delegate when scene playing time reaches the length of the longest playing montage. This is useful when some montage never ends (i.e. bEnableAutoBlendOut=false)
- ScenePivots have now their own dedicated type `FContextualAnimSetPivot` and are no longer computed on Start by the AnimSceneInstance. They can now be computed externally from the bindings and must be set before calling Start. This is to support replicated playback since pivots are replicated from the server.
- updated FContextualAnimStartSceneParams to store information regarding the `Section` (as for the AnimSet)
#rnx

diff --git a/Engine/Plugins/Experimental/Animation/ContextualAnimation/Source/ContextualAnimation/Private/ContextualAnimManager.cpp b/Engine/Plugins/Experimental/Animation/ContextualAnimation/Source/ContextualAnimation/Private/ContextualAnimManager.cpp
--- a/Engine/Plugins/Experimental/Animation/ContextualAnimation/Source/ContextualAnimation/Private/ContextualAnimManager.cpp
+++ b/Engine/Plugins/Experimental/Animation/ContextualAnimation/Source/ContextualAnimation/Private/ContextualAnimManager.cpp
@@ -1,12 +1,11 @@
 // Copyright Epic Games, Inc. All Rights Reserved.
 
 #include "ContextualAnimManager.h"
-#include "ContextualAnimActorInterface.h"
 #include "ContextualAnimSceneInstance.h"
-#include "DrawDebugHelpers.h"
 #include "ContextualAnimSceneAsset.h"
 #include "ContextualAnimation.h"
 #include "ContextualAnimSceneActorComponent.h"
+#include "ContextualAnimUtilities.h"
 #include "Engine/World.h"
 #include "GameFramework/Actor.h"
 
@@ -60,7 +59,9 @@
 
 void UContextualAnimManager::Tick(float DeltaTime)
 {
-	for (UContextualAnimSceneInstance* SceneInstance : Instances)
+	// Keep local copy since an instance might get unregistered by its tick
+	TArray<UContextualAnimSceneInstance*> InstancesToTick = Instances;
+	for (UContextualAnimSceneInstance* SceneInstance : InstancesToTick)
 	{
 		SceneInstance->Tick(DeltaTime);
 	}
@@ -70,7 +71,7 @@
 {
 	if (Actor)
 	{
-		for (UContextualAnimSceneInstance* SceneInstance : Instances)
+		for (const UContextualAnimSceneInstance* SceneInstance : Instances)
 		{
 			if (SceneInstance->IsActorInThisScene(Actor))
 			{
@@ -106,12 +107,13 @@
 		return nullptr;
 	}
 
+	const int32 SectionIdx = Params.SectionIdx;
 	FContextualAnimSceneBindings Bindings;
 	for (const auto& Pair : Params.RoleToActorMap)
 	{
 		FName RoleToBind = Pair.Key;
 
-		AActor* ActorToBind = Pair.Value.GetActor();
+		const AActor* ActorToBind = Pair.Value.GetActor();
 		if (ActorToBind == nullptr)
 		{
 			UE_LOG(LogContextualAnim, Warning, TEXT("UContextualAnimManager::ForceStartScene. Can't start scene. Reason: Trying to bind Invalid Actor. SceneAsset: %s Role: %s"),
@@ -120,7 +122,6 @@
 			return nullptr;
 		}
 
-		const int32 SectionIdx = 0; // Always start from the first section
 		const FContextualAnimTrack* AnimTrack = SceneAsset.GetAnimTrack(SectionIdx, Params.AnimSetIdx, RoleToBind);
 		if (AnimTrack == nullptr)
 		{
@@ -133,10 +134,18 @@
 		Bindings.Add(FContextualAnimSceneBinding(Pair.Value, SceneAsset, *AnimTrack));
 	}
 
-	UClass* Class = SceneAsset.GetSceneInstanceClass();
+	const UClass* Class = SceneAsset.GetSceneInstanceClass();
 	UContextualAnimSceneInstance* NewInstance = Class ? NewObject<UContextualAnimSceneInstance>(this, Class) : NewObject<UContextualAnimSceneInstance>(this);
 	NewInstance->SceneAsset = &SceneAsset;
 	NewInstance->Bindings = MoveTemp(Bindings);
+	if (Params.Pivots.IsEmpty())
+	{
+		Bindings.CalculateAnimSetPivots(NewInstance->GetMutablePivots());
+	}
+	else
+	{
+		NewInstance->SetPivots(Params.Pivots);
+	}
 	NewInstance->Start();
 	NewInstance->OnSceneEnded.AddDynamic(this, &UContextualAnimManager::OnSceneInstanceEnded);
 
@@ -158,7 +167,7 @@
 
 UContextualAnimSceneInstance* UContextualAnimManager::TryStartScene(const UContextualAnimSceneAsset& SceneAsset, const FContextualAnimStartSceneParams& Params)
 {
-	const int32 SectionIdx = 0; // Always start from the first section
+	const int32 SectionIdx = Params.SectionIdx;
 	FContextualAnimSceneBindings Bindings;
 
 	bool bSuccess = false;
@@ -181,10 +190,18 @@
 
 	if (bSuccess)
 	{
-		UClass* Class = SceneAsset.GetSceneInstanceClass();
+		const UClass* Class = SceneAsset.GetSceneInstanceClass();
 		UContextualAnimSceneInstance* NewInstance = Class ? NewObject<UContextualAnimSceneInstance>(this, Class) : NewObject<UContextualAnimSceneInstance>(this);
 		NewInstance->SceneAsset = &SceneAsset;
 		NewInstance->Bindings = MoveTemp(Bindings);
+		if (Params.Pivots.IsEmpty())
+		{
+			Bindings.CalculateAnimSetPivots(NewInstance->GetMutablePivots());
+		}
+		else
+		{
+			NewInstance->SetPivots(Params.Pivots);
+		}
 		NewInstance->Start();
 		NewInstance->OnSceneEnded.AddDynamic(this, &UContextualAnimManager::OnSceneInstanceEnded);
 
diff --git a/Engine/Plugins/Experimental/Animation/ContextualAnimation/Source/ContextualAnimation/Private/ContextualAnimSceneAsset.cpp b/Engine/Plugins/Experimental/Animation/ContextualAnimation/Source/ContextualAnimation/Private/ContextualAnimSceneAsset.cpp
--- a/Engine/Plugins/Experimental/Animation/ContextualAnimation/Source/ContextualAnimation/Private/ContextualAnimSceneAsset.cpp
+++ b/Engine/Plugins/Experimental/Animation/ContextualAnimation/Source/ContextualAnimation/Private/ContextualAnimSceneAsset.cpp
@@ -511,6 +511,11 @@
 	return Sections.FindByPredicate([&SectionName](const FContextualAnimSceneSection& Section) { return Section.Name == SectionName; });
 }
 
+int32 UContextualAnimSceneAsset::GetSectionIndex(const FName& SectionName) const
+{
+	return Sections.IndexOfByPredicate([&SectionName](const FContextualAnimSceneSection& Section) { return Section.Name == SectionName; });
+}
+
 const FContextualAnimTrack* UContextualAnimSceneAsset::GetAnimTrack(int32 SectionIdx, int32 AnimSetIdx, const FName& Role) const
 {
 	return (Sections.IsValidIndex(SectionIdx)) ? Sections[SectionIdx].GetAnimTrack(AnimSetIdx, Role) : nullptr;
diff --git a/Engine/Plugins/Experimental/Animation/ContextualAnimation/Source/ContextualAnimation/Private/ContextualAnimSceneInstance.cpp b/Engine/Plugins/Experimental/Animation/ContextualAnimation/Source/ContextualAnimation/Private/ContextualAnimSceneInstance.cpp
--- a/Engine/Plugins/Experimental/Animation/ContextualAnimation/Source/ContextualAnimation/Private/ContextualAnimSceneInstance.cpp
+++ b/Engine/Plugins/Experimental/Animation/ContextualAnimation/Source/ContextualAnimation/Private/ContextualAnimSceneInstance.cpp
@@ -7,12 +7,9 @@
 #include "GameFramework/Character.h"
 #include "GameFramework/CharacterMovementComponent.h"
 #include "MotionWarpingComponent.h"
-#include "ContextualAnimActorInterface.h"
 #include "ContextualAnimSceneAsset.h"
 #include "ContextualAnimSceneActorComponent.h"
-#include "ContextualAnimUtilities.h"
 #include "Engine/World.h"
-#include "DrawDebugHelpers.h"
 
 // UContextualAnimSceneInstance
 //================================================================================================================
@@ -28,8 +25,14 @@
 	return GetOuter() ? GetOuter()->GetWorld() : nullptr;
 }
 
-void UContextualAnimSceneInstance::Tick(float DeltaTime)
+void UContextualAnimSceneInstance::Tick(const float DeltaTime)
 {
+	RemainingDuration -= DeltaTime;
+	if (RemainingDuration <= 0.f)
+	{
+		OnSectionEndTimeReached.Broadcast(this);
+		RemainingDuration = MAX_flt;
+	}
 }
 
 bool UContextualAnimSceneInstance::IsActorInThisScene(const AActor* Actor) const
@@ -64,23 +67,36 @@
 	}
 }
 
-void UContextualAnimSceneInstance::Join(FContextualAnimSceneBinding& Binding)
+float UContextualAnimSceneInstance::Join(FContextualAnimSceneBinding& Binding)
 {
+	float Duration = MIN_flt;
+
 	AActor* Actor = Binding.GetActor();
 	if (Actor == nullptr)
 	{
-		return;
+		return Duration;
 	}
 
 	if (UAnimSequenceBase* Animation = Binding.GetAnimTrack().Animation)
 	{
 		if (UAnimInstance* AnimInstance = Binding.GetAnimInstance())
 		{
-			if(UAnimMontage* Montage = PlayAnimation(*AnimInstance, *Animation))
+			if (const UAnimMontage* Montage = PlayAnimation(*AnimInstance, *Animation))
 			{
 				AnimInstance->OnPlayMontageNotifyBegin.AddUniqueDynamic(this, &UContextualAnimSceneInstance::OnNotifyBeginReceived);
 				AnimInstance->OnPlayMontageNotifyEnd.AddUniqueDynamic(this, &UContextualAnimSceneInstance::OnNotifyEndReceived);
 				AnimInstance->OnMontageBlendingOut.AddUniqueDynamic(this, &UContextualAnimSceneInstance::OnMontageBlendingOut);
+
+				const float AdjustedPlayRate = AnimInstance->Montage_GetPlayRate(Montage) * Montage->RateScale;				
+				if (AdjustedPlayRate > 0.f)
+				{
+					Duration = (Montage->GetPlayLength() / AdjustedPlayRate);	
+				}
+				else
+				{
+					UE_LOG(LogContextualAnim, Warning, TEXT("Undesired playrate %.3f, using montage play length instead."), AdjustedPlayRate);
+					Duration = Montage->GetPlayLength();
+				}
 			}
 		}
 
@@ -96,37 +112,37 @@
 
 	if (UMotionWarpingComponent* MotionWarpComp = Actor->FindComponentByClass<UMotionWarpingComponent>())
 	{
-		for (const auto& Pair : AlignmentSectionToScenePivotList)
+		for (const FContextualAnimSetPivot& Pivot : AlignmentSectionToScenePivotList)
 		{
-			const FName WarpTargetName = Pair.Key;
-			const FTransform ScenePivotRuntime = Pair.Value;
-			const float Time = Binding.GetAnimTrack().GetSyncTimeForWarpSection(WarpTargetName);
-			const FTransform TransformRelativeToScenePivot = Binding.GetAnimTrack().AlignmentData.ExtractTransformAtTime(WarpTargetName, Time);
-			const FTransform WarpTarget = TransformRelativeToScenePivot * ScenePivotRuntime;
-			MotionWarpComp->AddOrUpdateWarpTargetFromTransform(WarpTargetName, WarpTarget);
+			const float Time = Binding.GetAnimTrack().GetSyncTimeForWarpSection(Pivot.Name);
+			const FTransform TransformRelativeToScenePivot = Binding.GetAnimTrack().AlignmentData.ExtractTransformAtTime(Pivot.Name, Time);
+			const FTransform WarpTarget = TransformRelativeToScenePivot * Pivot.Transform;
+			MotionWarpComp->AddOrUpdateWarpTargetFromTransform(Pivot.Name, WarpTarget);
 		}
 	}
 
-	if(SceneAsset->GetDisableCollisionBetweenActors())
+	if (SceneAsset->GetDisableCollisionBetweenActors())
 	{
 		SetIgnoreCollisionWithOtherActors(Actor, true);
 	}
 
 	Binding.SceneInstancePtr = this;
 
-	if(UContextualAnimSceneActorComponent* SceneActorComp = Binding.GetSceneActorComponent())
+	if (UContextualAnimSceneActorComponent* SceneActorComp = Binding.GetSceneActorComponent())
 	{
 		SceneActorComp->OnJoinedScene(&Binding);
 	}
 
 	OnActorJoined.Broadcast(this, Actor);
+
+	return Duration;
 }
 
 void UContextualAnimSceneInstance::Leave(FContextualAnimSceneBinding& Binding)
 {
 	if (UAnimInstance* AnimInstance = Binding.GetAnimInstance())
 	{
-		if (UAnimMontage* CurrentMontage = AnimInstance->GetCurrentActiveMontage())
+		if (const UAnimMontage* CurrentMontage = AnimInstance->GetCurrentActiveMontage())
 		{
 			AnimInstance->Montage_Stop(CurrentMontage->BlendOut.GetBlendTime(), CurrentMontage);
 		}
@@ -158,17 +174,12 @@
 
 void UContextualAnimSceneInstance::Start()
 {
-	const int32 SectionIdx = 0; // Always start from the first section
-	for (const FContextualAnimSetPivotDefinition& Def : SceneAsset->GetAnimSetPivotDefinitionsInSection(SectionIdx))
-	{
-		FTransform ScenePivotRuntime = FTransform::Identity;
-		UContextualAnimUtilities::CalculateAnimSetPivot(Def, Bindings, ScenePivotRuntime);
-		AlignmentSectionToScenePivotList.Add(MakeTuple(Def.Name, ScenePivotRuntime));
-	}
+	RemainingDuration = 0.f;
 
 	for (auto& Binding : Bindings)
 	{
-		Join(Binding);
+		const float TrackDuration = Join(Binding);
+		RemainingDuration = FMath::Max(RemainingDuration, TrackDuration);
 	}
 }
 
@@ -226,7 +237,7 @@
 		if (UAnimInstance* AnimInstance = Binding.GetAnimInstance())
 		{
 			// Keep montage support for now but might go away soon
-			if (UAnimMontage* AnimMontage = Cast<UAnimMontage>(Binding.GetAnimTrack().Animation))
+			if (const UAnimMontage* AnimMontage = Cast<UAnimMontage>(Binding.GetAnimTrack().Animation))
 			{
 				if (AnimInstance->Montage_IsPlaying(AnimMontage))
 				{
@@ -255,9 +266,10 @@
 		}
 	}
 
-	if(bShouldEnd)
+	if (bShouldEnd)
 	{
 		OnSceneEnded.Broadcast(this);
+		OnSectionDonePlaying.Broadcast(this);
 	}
 }
 
diff --git a/Engine/Plugins/Experimental/Animation/ContextualAnimation/Source/ContextualAnimation/Private/ContextualAnimTypes.cpp b/Engine/Plugins/Experimental/Animation/ContextualAnimation/Source/ContextualAnimation/Private/ContextualAnimTypes.cpp
--- a/Engine/Plugins/Experimental/Animation/ContextualAnimation/Source/ContextualAnimation/Private/ContextualAnimTypes.cpp
+++ b/Engine/Plugins/Experimental/Animation/ContextualAnimation/Source/ContextualAnimation/Private/ContextualAnimTypes.cpp
@@ -366,3 +366,43 @@
 	// Success if all the roles were filled
 	return (OutBindings.Num() == SceneAsset.GetNumRoles());
 }
+
+void FContextualAnimSceneBindings::CalculateAnimSetPivots(TArray<FContextualAnimSetPivot>& OutScenePivots) const
+{
+	if (const UContextualAnimSceneAsset* SceneAsset = GetSceneAsset())
+	{
+		const int32 SectionIdx = GetSectionIdx();
+		for (const FContextualAnimSetPivotDefinition& Def : SceneAsset->GetAnimSetPivotDefinitionsInSection(SectionIdx))
+		{
+			FContextualAnimSetPivot& ScenePivotRuntime = OutScenePivots.AddDefaulted_GetRef();
+			CalculateAnimSetPivot(Def, ScenePivotRuntime);
+		}
+	}
+}
+
+bool FContextualAnimSceneBindings::CalculateAnimSetPivot(const FContextualAnimSetPivotDefinition& AnimSetPivotDef, FContextualAnimSetPivot& OutScenePivot) const
+{
+	if (const FContextualAnimSceneBinding* Binding = FindBindingByRole(AnimSetPivotDef.Origin))
+	{
+		OutScenePivot.Name = AnimSetPivotDef.Name;
+		if (AnimSetPivotDef.bAlongClosestDistance)
+		{
+			if (const FContextualAnimSceneBinding* OtherBinding = FindBindingByRole(AnimSetPivotDef.OtherRole))
+			{
+				const FTransform T1 = Binding->GetTransform();
+				const FTransform T2 = OtherBinding->GetTransform();
+
+				OutScenePivot.Transform.SetLocation(FMath::Lerp<FVector>(T1.GetLocation(), T2.GetLocation(), AnimSetPivotDef.Weight));
+				OutScenePivot.Transform.SetRotation((T2.GetLocation() - T1.GetLocation()).GetSafeNormal2D().ToOrientationQuat());
+				return true;
+			}
+		}
+		else
+		{
+			OutScenePivot.Transform = Binding->GetTransform();
+			return true;
+		}
+	}
+
+	return false;
+}
diff --git a/Engine/Plugins/Experimental/Animation/ContextualAnimation/Source/ContextualAnimation/Private/ContextualAnimUtilities.cpp b/Engine/Plugins/Experimental/Animation/ContextualAnimation/Source/ContextualAnimation/Private/ContextualAnimUtilities.cpp
--- a/Engine/Plugins/Experimental/Animation/ContextualAnimation/Source/ContextualAnimation/Private/ContextualAnimUtilities.cpp
+++ b/Engine/Plugins/Experimental/Animation/ContextualAnimation/Source/ContextualAnimation/Private/ContextualAnimUtilities.cpp
@@ -279,33 +279,6 @@
 	return false;
 }
 
-
-bool UContextualAnimUtilities::CalculateAnimSetPivot(const FContextualAnimSetPivotDefinition& AnimSetPivotDef, const FContextualAnimSceneBindings& Bindings, FTransform& OutScenePivot)
-{
-	if (const FContextualAnimSceneBinding* Binding = Bindings.FindBindingByRole(AnimSetPivotDef.Origin))
-	{
-		if (AnimSetPivotDef.bAlongClosestDistance)
-		{
-			if (const FContextualAnimSceneBinding* OtherBinding = Bindings.FindBindingByRole(AnimSetPivotDef.OtherRole))
-			{
-				const FTransform T1 = Binding->GetTransform();
-				const FTransform T2 = OtherBinding->GetTransform();
-
-				OutScenePivot.SetLocation(FMath::Lerp<FVector>(T1.GetLocation(), T2.GetLocation(), AnimSetPivotDef.Weight));
-				OutScenePivot.SetRotation((T2.GetLocation() - T1.GetLocation()).GetSafeNormal2D().ToOrientationQuat());
-				return true;
-			}
-		}
-		else
-		{
-			OutScenePivot = Binding->GetTransform();
-			return true;
-		}
-	}
-
-	return false;
-}
-
 // SceneBindings Blueprint Interface
 //------------------------------------------------------------------------------------------
 
@@ -319,6 +292,11 @@
 	return FContextualAnimSceneBinding::InvalidBinding;
 }
 
+void UContextualAnimUtilities::BP_SceneBindings_CalculateAnimSetPivots(const FContextualAnimSceneBindings& Bindings, TArray<FContextualAnimSetPivot>& OutPivots)
+{
+	Bindings.CalculateAnimSetPivots(OutPivots);
+}
+
 void UContextualAnimUtilities::BP_SceneBindings_AddOrUpdateWarpTargetsForBindings(const FContextualAnimSceneBindings& Bindings)
 {
 	const UContextualAnimSceneAsset* SceneAsset = Bindings.GetSceneAsset();
@@ -327,8 +305,8 @@
 		const int32 SectionIdx = Bindings.GetSectionIdx();
 		for (const FContextualAnimSetPivotDefinition& PivotDef : SceneAsset->GetAnimSetPivotDefinitionsInSection(SectionIdx))
 		{
-			FTransform ScenePivot = FTransform::Identity;
-			if (UContextualAnimUtilities::CalculateAnimSetPivot(PivotDef, Bindings, ScenePivot))
+			FContextualAnimSetPivot ScenePivot;
+			if (Bindings.CalculateAnimSetPivot(PivotDef, ScenePivot))
 			{
 				for (const FContextualAnimSceneBinding& Binding : Bindings)
 				{
@@ -337,7 +315,7 @@
 						//@TODO: Cache this
 						const float Time = Binding.GetAnimTrack().GetSyncTimeForWarpSection(PivotDef.Name);
 						const FTransform TransformRelativeToScenePivot = Binding.GetAnimTrack().AlignmentData.ExtractTransformAtTime(PivotDef.Name, Time);
-						const FTransform WarpTarget = (TransformRelativeToScenePivot * ScenePivot);
+						const FTransform WarpTarget = (TransformRelativeToScenePivot * ScenePivot.Transform);
 						MotionWarpComp->AddOrUpdateWarpTargetFromTransform(PivotDef.Name, WarpTarget);
 					}
 				}
diff --git a/Engine/Plugins/Experimental/Animation/ContextualAnimation/Source/ContextualAnimation/Public/ContextualAnimSceneAsset.h b/Engine/Plugins/Experimental/Animation/ContextualAnimation/Source/ContextualAnimation/Public/ContextualAnimSceneAsset.h
--- a/Engine/Plugins/Experimental/Animation/ContextualAnimation/Source/ContextualAnimation/Public/ContextualAnimSceneAsset.h
+++ b/Engine/Plugins/Experimental/Animation/ContextualAnimation/Source/ContextualAnimation/Public/ContextualAnimSceneAsset.h
@@ -141,6 +141,8 @@
 
 	const FContextualAnimSceneSection* GetSection(const FName& SectionName) const;
 
+	int32 GetSectionIndex(const FName& SectionName) const;
+	
  	const FContextualAnimTrack* GetAnimTrack(int32 SectionIdx, int32 AnimSetIdx, const FName& Role) const;
  
  	const FContextualAnimTrack* GetAnimTrack(int32 SectionIdx, int32 AnimSetIdx, int32 AnimTrackIdx) const;
diff --git a/Engine/Plugins/Experimental/Animation/ContextualAnimation/Source/ContextualAnimation/Public/ContextualAnimSceneInstance.h b/Engine/Plugins/Experimental/Animation/ContextualAnimation/Source/ContextualAnimation/Public/ContextualAnimSceneInstance.h
--- a/Engine/Plugins/Experimental/Animation/ContextualAnimation/Source/ContextualAnimation/Public/ContextualAnimSceneInstance.h
+++ b/Engine/Plugins/Experimental/Animation/ContextualAnimation/Source/ContextualAnimation/Public/ContextualAnimSceneInstance.h
@@ -42,6 +42,20 @@
 	friend class UContextualAnimManager;
 	friend class FContextualAnimViewModel;
 
+	DECLARE_MULTICAST_DELEGATE_OneParam(FGenericSceneMulticastDelegate, UContextualAnimSceneInstance*);
+
+	/**
+	 * Delegate to notify after all montages played by the scene section blended out.
+	 * Will not be broadcasted if one or more montages have 'bEnableAutoBlendOut' set to 'false'.
+	 */
+	FGenericSceneMulticastDelegate OnSectionDonePlaying;
+
+	/**
+	 * Delegate to notify once the scene play time reaches the duration defined by the longest played montage of the selected section.
+	 * This delegate should be used if one or more montages have 'bEnableAutoBlendOut' set to 'false'.
+	 */
+	FGenericSceneMulticastDelegate OnSectionEndTimeReached;
+
 	/** Delegate to notify external objects when this is scene is completed */
 	UPROPERTY(BlueprintAssignable)
 	FOnContextualAnimSceneEnded OnSceneEnded;
@@ -78,8 +92,8 @@
 	bool IsActorInThisScene(const AActor* Actor) const;
 
 	const UContextualAnimSceneAsset& GetSceneAsset() const { return *SceneAsset; }
-	const FContextualAnimSceneBindings& GetBindings() const { return Bindings; };
-	FContextualAnimSceneBindings& GetBindings() { return Bindings; };
+	const FContextualAnimSceneBindings& GetBindings() const { return Bindings; }
+	FContextualAnimSceneBindings& GetBindings() { return Bindings; }
 	const FContextualAnimSceneBinding* FindBindingByActor(const AActor* Actor) const { return Bindings.FindBindingByActor(Actor); }
 	const FContextualAnimSceneBinding* FindBindingByRole(const FName& Role) const { return Bindings.FindBindingByRole(Role); }
 
@@ -88,14 +102,20 @@
 
 protected:
 
-	/** Tells the scene actor to join the scene (play animation) */
-	void Join(FContextualAnimSceneBinding& Binding);
+	/**
+	 * Tells the scene actor to join the scene (play animation)
+	 * @return Duration of the playing animation, or MIN_flt if not playing one.
+	 */
+	float Join(FContextualAnimSceneBinding& Binding);
 
 	/** Tells the scene actor to leave the scene (stop animation) */
 	void Leave(FContextualAnimSceneBinding& Binding);
 
 	bool TransitionTo(FContextualAnimSceneBinding& Binding, const FContextualAnimTrack& AnimTrack);
 
+	TArray<FContextualAnimSetPivot>& GetMutablePivots() { return AlignmentSectionToScenePivotList; }
+	void SetPivots(const TArray<FContextualAnimSetPivot>& Pivots) { AlignmentSectionToScenePivotList = Pivots; }
+
 	/** Helper function to set ignore collision between the supplied actor and all the other actors in this scene */
 	void SetIgnoreCollisionWithOtherActors(AActor* Actor, bool bValue) const;
 
@@ -117,8 +137,14 @@
 	UPROPERTY()
 	FContextualAnimSceneBindings Bindings;
 
-	TArray<TTuple<FName, FTransform>> AlignmentSectionToScenePivotList;
+	TArray<FContextualAnimSetPivot> AlignmentSectionToScenePivotList;
 
+	/**
+	 * Remaining scene section duration initially computed based on the longest animation duration from all actors that joined the scene.
+	 * Delegate 'OnSectionEndTimeReached' is broadcasted once value reaches 0.
+	 */
+	float RemainingDuration;
+
 	/** Helper to play an AnimSequenceBase as montage. If Animation is not a montage it plays it as dynamic montage  */
 	UAnimMontage* PlayAnimation(UAnimInstance& AnimInstance, UAnimSequenceBase& Animation);
 };
diff --git a/Engine/Plugins/Experimental/Animation/ContextualAnimation/Source/ContextualAnimation/Public/ContextualAnimTypes.h b/Engine/Plugins/Experimental/Animation/ContextualAnimation/Source/ContextualAnimation/Public/ContextualAnimTypes.h
--- a/Engine/Plugins/Experimental/Animation/ContextualAnimation/Source/ContextualAnimation/Public/ContextualAnimTypes.h
+++ b/Engine/Plugins/Experimental/Animation/ContextualAnimation/Source/ContextualAnimation/Public/ContextualAnimTypes.h
@@ -240,6 +240,25 @@
 	float Weight = 0.f;
 };
 
+// FContextualAnimSetPivot
+///////////////////////////////////////////////////////////////////////
+
+/** Pivot for a AnimSet */
+USTRUCT(BlueprintType)
+struct FContextualAnimSetPivot
+{
+	GENERATED_BODY()
+
+	FContextualAnimSetPivot() = default;
+	FContextualAnimSetPivot(const FName InName, const FTransform& InTransform) : Name(InName), Transform(InTransform) {}
+	
+	UPROPERTY(EditAnywhere, BlueprintReadWrite, Category = "Defaults")
+	FName Name = NAME_None;
+
+	UPROPERTY(EditAnywhere, BlueprintReadWrite, Category = "Defaults")
+	FTransform Transform;
+};
+
 // FContextualAnimSceneBindingContext
 ///////////////////////////////////////////////////////////////////////
 
@@ -374,6 +393,9 @@
 
 	static bool TryCreateBindings(const UContextualAnimSceneAsset& SceneAsset, int32 SectionIdx, int32 AnimSetIdx, const TMap<FName, FContextualAnimSceneBindingContext>& Params, FContextualAnimSceneBindings& OutBindings);
 
+	void CalculateAnimSetPivots(TArray<FContextualAnimSetPivot>& OutScenePivots) const;
+	bool CalculateAnimSetPivot(const FContextualAnimSetPivotDefinition& AnimSetPivotDef, FContextualAnimSetPivot& OutScenePivot) const;
+
 private:
 
 	/** List of actors bound to each role in the SceneAsset */
@@ -382,7 +404,7 @@
 };
 
 USTRUCT(BlueprintType)
-struct FContextualAnimStartSceneParams
+struct CONTEXTUALANIMATION_API FContextualAnimStartSceneParams
 {
 	GENERATED_BODY()
 
@@ -390,14 +412,27 @@
 	UPROPERTY(EditAnywhere, BlueprintReadWrite, Category = "Defaults")
 	TMap<FName, FContextualAnimSceneBindingContext> RoleToActorMap;
 
-	/** Desired set. If INDEX_NONE the Manager will attempt to find the best set to use (in the first section) by running the selection criteria */
+	/** Desired section. If INDEX_NONE the Manager will use or find best set in the first section. */
+	UPROPERTY(EditAnywhere, BlueprintReadWrite, Category = "Defaults")
+	int32 SectionIdx = INDEX_NONE;
+
+	/** Desired set. If INDEX_NONE the Manager will attempt to find the best set to use by running the selection criteria.
+	 * The selection will be performed in the section specified by SectionIdx or in the first section if SectionIdx == INDEX_NONE.
+	 */
 	UPROPERTY(EditAnywhere, BlueprintReadWrite, Category = "Defaults")
 	int32 AnimSetIdx = INDEX_NONE;
 
+	/** Precomputed pivots that could be provided when starting a scene.
+	 * When not provided, the pivots will be automatically computed using local context information.
+	 */
+	TArray<FContextualAnimSetPivot> Pivots;
+
 	void Reset()
 	{
 		RoleToActorMap.Reset();
+		SectionIdx = INDEX_NONE;
 		AnimSetIdx = INDEX_NONE;
+		Pivots.Reset();
 	}
 };
 
diff --git a/Engine/Plugins/Experimental/Animation/ContextualAnimation/Source/ContextualAnimation/Public/ContextualAnimUtilities.h b/Engine/Plugins/Experimental/Animation/ContextualAnimation/Source/ContextualAnimation/Public/ContextualAnimUtilities.h
--- a/Engine/Plugins/Experimental/Animation/ContextualAnimation/Source/ContextualAnimation/Public/ContextualAnimUtilities.h
+++ b/Engine/Plugins/Experimental/Animation/ContextualAnimation/Source/ContextualAnimation/Public/ContextualAnimUtilities.h
@@ -54,8 +54,6 @@
 
 	static void DrawSector(FPrimitiveDrawInterface& PDI, const FVector& Origin, const FVector& Direction, float MinDistance, float MaxDistance, float MinAngle, float MaxAngle, const FLinearColor& Color, uint8 DepthPriority, float Thickness);
 
-	static bool CalculateAnimSetPivot(const FContextualAnimSetPivotDefinition& AnimSetPivotDef, const FContextualAnimSceneBindings& Bindings, FTransform& OutScenePivot);
-
 	UFUNCTION(BlueprintCallable, Category = "Contextual Anim|Scene Asset", meta = (DisplayName = "Create Contextual Anim Scene Bindings"))
 	static bool BP_CreateContextualAnimSceneBindings(const UContextualAnimSceneAsset* SceneAsset, const TMap<FName, FContextualAnimSceneBindingContext>& Params, FContextualAnimSceneBindings& OutBindings);
 
@@ -74,6 +72,9 @@
 	// SceneBindings Blueprint Interface
 	//------------------------------------------------------------------------------------------
 
+	UFUNCTION(BlueprintCallable, Category = "Contextual Anim|Scene Bindings", meta = (DisplayName = "Calculate pivots For Bindings"))
+	static void BP_SceneBindings_CalculateAnimSetPivots(const FContextualAnimSceneBindings& Bindings, TArray<FContextualAnimSetPivot>& OutPivots);
+	
 	UFUNCTION(BlueprintCallable, Category = "Contextual Anim|Scene Bindings", meta = (DisplayName = "Add Or Update Warp Targets For Bindings"))
 	static void BP_SceneBindings_AddOrUpdateWarpTargetsForBindings(const FContextualAnimSceneBindings& Bindings);
 
diff --git a/Engine/Plugins/Experimental/Animation/ContextualAnimation/Source/ContextualAnimationEditor/Private/ContextualAnimViewModel.cpp b/Engine/Plugins/Experimental/Animation/ContextualAnimation/Source/ContextualAnimationEditor/Private/ContextualAnimViewModel.cpp
--- a/Engine/Plugins/Experimental/Animation/ContextualAnimation/Source/ContextualAnimationEditor/Private/ContextualAnimViewModel.cpp
+++ b/Engine/Plugins/Experimental/Animation/ContextualAnimation/Source/ContextualAnimationEditor/Private/ContextualAnimViewModel.cpp
@@ -192,6 +192,7 @@
 	Sequencer->NotifyMovieSceneDataChanged(EMovieSceneDataChangeType::MovieSceneStructureItemsChanged);
 
 	StartSceneParams.Reset();
+	StartSceneParams.SectionIdx = 0;
 	StartSceneParams.AnimSetIdx = 0;
 
 	TArray<FName> Roles = SceneAsset->GetRoles();
