Fix static analysis warning (null check is redundant because pointer has already been dereferenced).

This change should not make any difference to behavior, but may not be what the code was originally intended to do. Please check.

diff --git a/Engine/Plugins/Experimental/Animation/ContextualAnimation/Source/ContextualAnimationEditor/Private/ContextualAnimViewModel.cpp b/Engine/Plugins/Experimental/Animation/ContextualAnimation/Source/ContextualAnimationEditor/Private/ContextualAnimViewModel.cpp
--- a/Engine/Plugins/Experimental/Animation/ContextualAnimation/Source/ContextualAnimationEditor/Private/ContextualAnimViewModel.cpp
+++ b/Engine/Plugins/Experimental/Animation/ContextualAnimation/Source/ContextualAnimationEditor/Private/ContextualAnimViewModel.cpp
@@ -493,7 +493,7 @@ void FContextualAnimViewModel::SequencerTimeChanged()
 			{
 				bShouldTransition = FocusedAnimAsMontage != MontageInstance->Montage;
 			}
-			else if(FocusedAnim)
+			else
 			{
 				bShouldTransition = FocusedAnim != MontageInstance->Montage->SlotAnimTracks[0].AnimTrack.AnimSegments[0].GetAnimReference();
 			}
