Fix warning related to shortname in QAGame

diff --git a/Engine/Plugins/Experimental/Animation/ContextualAnimation/Source/ContextualAnimation/Public/ContextualAnimTypes.h b/Engine/Plugins/Experimental/Animation/ContextualAnimation/Source/ContextualAnimation/Public/ContextualAnimTypes.h
--- a/Engine/Plugins/Experimental/Animation/ContextualAnimation/Source/ContextualAnimation/Public/ContextualAnimTypes.h
+++ b/Engine/Plugins/Experimental/Animation/ContextualAnimation/Source/ContextualAnimation/Public/ContextualAnimTypes.h
@@ -262,7 +262,7 @@ struct FContextualAnimSetPivot
 // FContextualAnimSceneBindingContext
 ///////////////////////////////////////////////////////////////////////
 
-USTRUCT(BlueprintType, meta = (HasNativeMake = "ContextualAnimUtilities.BP_SceneBindingContext_MakeFromActor"))
+USTRUCT(BlueprintType, meta = (HasNativeMake = "/Script/ContextualAnimation.ContextualAnimUtilities:BP_SceneBindingContext_MakeFromActor"))
 struct CONTEXTUALANIMATION_API FContextualAnimSceneBindingContext
 {
 	GENERATED_BODY()
diff --git a/Engine/Plugins/Runtime/OpenXR/Source/OpenXRInput/Public/OpenXRInputSettings.h b/Engine/Plugins/Runtime/OpenXR/Source/OpenXRInput/Public/OpenXRInputSettings.h
--- a/Engine/Plugins/Runtime/OpenXR/Source/OpenXRInput/Public/OpenXRInputSettings.h
+++ b/Engine/Plugins/Runtime/OpenXR/Source/OpenXRInput/Public/OpenXRInputSettings.h
@@ -20,6 +20,6 @@ public:
 	bool bUseEnhancedInput = false;
 
 	/** A the mappable input config used to generate action sets for OpenXR. */
-	UPROPERTY(config, EditAnywhere, Category = "Enhanced Input", meta = (DisplayName = "Mappable Input Config for XR", AllowedClasses = "PlayerMappableInputConfig"))
+	UPROPERTY(config, EditAnywhere, Category = "Enhanced Input", meta = (DisplayName = "Mappable Input Config for XR", AllowedClasses = "/Script/EnhancedInput.PlayerMappableInputConfig"))
 	FSoftObjectPath MappableInputConfig = nullptr;
 };
