Contextual Anim Editor: Timeline now shows the data from a single section and added menu to jump between sections
#fly [REDACTED_HANDLE]

diff --git a/Engine/Plugins/Experimental/Animation/ContextualAnimation/Source/ContextualAnimationEditor/Private/ContextualAnimAssetEditorToolkit.cpp b/Engine/Plugins/Experimental/Animation/ContextualAnimation/Source/ContextualAnimationEditor/Private/ContextualAnimAssetEditorToolkit.cpp
--- a/Engine/Plugins/Experimental/Animation/ContextualAnimation/Source/ContextualAnimationEditor/Private/ContextualAnimAssetEditorToolkit.cpp
+++ b/Engine/Plugins/Experimental/Animation/ContextualAnimation/Source/ContextualAnimationEditor/Private/ContextualAnimAssetEditorToolkit.cpp
@@ -203,6 +203,14 @@
 
 	ToolbarBuilder.AddComboButton(
 		FUIAction(),
+		FOnGetContent::CreateSP(this, &FContextualAnimAssetEditorToolkit::BuildSectionsMenu),
+		LOCTEXT("Sections_Label", "Sections"),
+		FText::GetEmpty(),
+		FSlateIcon()
+	);
+
+	ToolbarBuilder.AddComboButton(
+		FUIAction(),
 		FOnGetContent::CreateSP(this, &FContextualAnimAssetEditorToolkit::BuildNewAnimSetWidget),
 		LOCTEXT("NewSet_Label", "New Set"),
 		FText::GetEmpty(),
@@ -218,6 +226,32 @@
 	);
 }
 
+TSharedRef<SWidget> FContextualAnimAssetEditorToolkit::BuildSectionsMenu()
+{
+	const bool bShouldCloseWindowAfterMenuSelection = true;
+	FMenuBuilder MenuBuilder(bShouldCloseWindowAfterMenuSelection, GetToolkitCommands());
+
+	MenuBuilder.BeginSection(NAME_None, LOCTEXT("Sections_Label", "Sections"));
+	{
+		const UContextualAnimSceneAsset* SceneAsset = GetSceneAsset();
+		
+		TArray<FName> SectionNames = SceneAsset->GetSectionNames();
+		for (int32 Idx = 0; Idx < SectionNames.Num(); Idx++)
+		{
+			MenuBuilder.AddMenuEntry(
+				FText::FromString(SectionNames[Idx].ToString()),
+				FText::GetEmpty(),
+				FSlateIcon(),
+				FUIAction(FExecuteAction::CreateLambda([this, Idx]() {
+					ViewModel->SetActiveSection(Idx);
+				})));
+		}
+	}
+	MenuBuilder.EndSection();
+
+	return MenuBuilder.MakeWidget();
+}
+
 TSharedRef<SWidget> FContextualAnimAssetEditorToolkit::BuildNewAnimSetWidget()
 {
 	const bool bShouldCloseWindowAfterMenuSelection = true;
diff --git a/Engine/Plugins/Experimental/Animation/ContextualAnimation/Source/ContextualAnimationEditor/Private/ContextualAnimViewModel.cpp b/Engine/Plugins/Experimental/Animation/ContextualAnimation/Source/ContextualAnimationEditor/Private/ContextualAnimViewModel.cpp
--- a/Engine/Plugins/Experimental/Animation/ContextualAnimation/Source/ContextualAnimationEditor/Private/ContextualAnimViewModel.cpp
+++ b/Engine/Plugins/Experimental/Animation/ContextualAnimation/Source/ContextualAnimationEditor/Private/ContextualAnimViewModel.cpp
@@ -105,6 +105,15 @@ void FContextualAnimViewModel::CreateSequencer()
 	Sequencer->SetPlaybackStatus(EMovieScenePlayerStatus::Stopped);
 }
 
+void FContextualAnimViewModel::SetActiveSection(int32 SectionIdx)
+{
+	check(GetSceneAsset()->Sections.IsValidIndex(SectionIdx));
+
+	ActiveSectionIdx = SectionIdx;
+
+	RefreshSequencerTracks();
+}
+
 void FContextualAnimViewModel::SetActiveAnimSetForSection(int32 SectionIdx, int32 AnimSetIdx)
 {
 	check(GetSceneAsset()->Sections.IsValidIndex(SectionIdx));
@@ -191,9 +200,14 @@ void FContextualAnimViewModel::RefreshSequencerTracks()
 
 	Sequencer->NotifyMovieSceneDataChanged(EMovieSceneDataChangeType::MovieSceneStructureItemsChanged);
 
+	if(!SceneAsset->Sections.IsValidIndex(ActiveSectionIdx))
+	{
+		return;
+	}
+
 	StartSceneParams.Reset();
-	StartSceneParams.SectionIdx = 0;
-	StartSceneParams.AnimSetIdx = 0;
+	StartSceneParams.SectionIdx = ActiveSectionIdx;
+	StartSceneParams.AnimSetIdx = ActiveAnimSetMap.FindOrAdd(ActiveSectionIdx);
 
 	TArray<FName> Roles = SceneAsset->GetRoles();
 	for(FName Role : Roles)
@@ -204,68 +218,53 @@ void FContextualAnimViewModel::RefreshSequencerTracks()
 		MovieSceneAnimTrack->Initialize(Role);
 	}
 
-	float SectionMaxTime = 0.f;
-	float SectionStartTime = 0.f;
-	for (int32 SectionIdx = 0; SectionIdx < SceneAsset->Sections.Num(); SectionIdx++)
+	FContextualAnimSceneSection& ContextualAnimSection = SceneAsset->Sections[ActiveSectionIdx];
+	for (int32 AnimSetIdx = 0; AnimSetIdx < ContextualAnimSection.AnimSets.Num(); AnimSetIdx++)
 	{
-		SectionStartTime = SectionMaxTime;
-
-		FContextualAnimSceneSection& ContextualAnimSection = SceneAsset->Sections[SectionIdx];
-		for (int32 AnimSetIdx = 0; AnimSetIdx < ContextualAnimSection.AnimSets.Num(); AnimSetIdx++)
+		FContextualAnimSet& AnimSet = ContextualAnimSection.AnimSets[AnimSetIdx];
+		for (int32 AnimTrackIdx = 0; AnimTrackIdx < AnimSet.Tracks.Num(); AnimTrackIdx++)
 		{
-			FContextualAnimSet& AnimSet = ContextualAnimSection.AnimSets[AnimSetIdx];
-			for (int32 AnimTrackIdx = 0; AnimTrackIdx < AnimSet.Tracks.Num(); AnimTrackIdx++)
+			const FContextualAnimTrack& AnimTrack = AnimSet.Tracks[AnimTrackIdx];
+
+			UAnimSequenceBase* Animation = AnimTrack.Animation;
+			if (Animation)
 			{
-				const FContextualAnimTrack& AnimTrack = AnimSet.Tracks[AnimTrackIdx];
-				
-				UAnimSequenceBase* Animation = AnimTrack.Animation;
-				if (Animation)
+				UContextualAnimMovieSceneTrack* MovieSceneTrack = FindMasterTrackByRole(AnimTrack.Role);
+				if (MovieSceneTrack == nullptr)
 				{
-					UContextualAnimMovieSceneTrack* MovieSceneTrack = FindMasterTrackByRole(AnimTrack.Role);
-					if (MovieSceneTrack == nullptr)
-					{
-						UE_LOG(LogContextualAnim, Warning, TEXT("FContextualAnimViewModel::RefreshSequencerTracks. Can't find MovieSceneTrack for %s. Role: %s SectionIdx: %d AnimIndex: %d"), 
-							*GetNameSafe(Animation), *AnimTrack.Role.ToString(), SectionIdx, AnimSetIdx);
-
-						continue;
-					}
+					UE_LOG(LogContextualAnim, Warning, TEXT("FContextualAnimViewModel::RefreshSequencerTracks. Can't find MovieSceneTrack for %s. Role: %s SectionIdx: %d AnimIndex: %d"),
+						*GetNameSafe(Animation), *AnimTrack.Role.ToString(), ActiveSectionIdx, AnimSetIdx);
 
-					UContextualAnimMovieSceneSection* NewSection = NewObject<UContextualAnimMovieSceneSection>(MovieSceneTrack, UContextualAnimMovieSceneSection::StaticClass(), NAME_None, RF_Transactional);
-					check(NewSection);
+					continue;
+				}
 
-					NewSection->Initialize(SectionIdx, AnimSetIdx, AnimTrackIdx);
+				UContextualAnimMovieSceneSection* NewSection = NewObject<UContextualAnimMovieSceneSection>(MovieSceneTrack, UContextualAnimMovieSceneSection::StaticClass(), NAME_None, RF_Transactional);
+				check(NewSection);
 
-					const float AnimLength = Animation->GetPlayLength();
-					const FFrameRate TickResolution = MovieSceneSequence->GetMovieScene()->GetTickResolution();
-					const FFrameNumber StartFrame = (SectionStartTime * TickResolution).RoundToFrame();
-					const FFrameNumber EndFrame = (AnimLength * TickResolution).RoundToFrame() + StartFrame;
-					NewSection->SetRange(TRange<FFrameNumber>::Inclusive(StartFrame, EndFrame));
+				NewSection->Initialize(ActiveSectionIdx, AnimSetIdx, AnimTrackIdx);
 
-					NewSection->SetRowIndex(AnimSetIdx);
+				const float AnimLength = Animation->GetPlayLength();
+				const FFrameRate TickResolution = MovieSceneSequence->GetMovieScene()->GetTickResolution();
+				NewSection->SetRange(TRange<FFrameNumber>::Inclusive(0, (AnimLength * TickResolution).RoundToFrame()));
 
-					const int32& ActiveSetIdx = ActiveAnimSetMap.FindOrAdd(SectionIdx);
-					NewSection->SetIsActive(AnimSetIdx == ActiveSetIdx);
+				NewSection->SetRowIndex(AnimSetIdx);
 
-					MovieSceneTrack->AddSection(*NewSection);
-					MovieSceneTrack->SetTrackRowDisplayName(FText::FromString(FString::Printf(TEXT("%d"), AnimSetIdx)), AnimSetIdx);
+				const int32& ActiveSetIdx = ActiveAnimSetMap.FindOrAdd(ActiveSectionIdx);
+				NewSection->SetIsActive(AnimSetIdx == ActiveSetIdx);
 
-					if (AnimLength > SectionMaxTime)
-					{
-						SectionMaxTime = AnimLength;
-					}
-				}
+				MovieSceneTrack->AddSection(*NewSection);
+				MovieSceneTrack->SetTrackRowDisplayName(FText::FromString(FString::Printf(TEXT("%d"), AnimSetIdx)), AnimSetIdx);
+			}
 
-				if (!StartSceneParams.RoleToActorMap.Contains(AnimTrack.Role))
+			if (!StartSceneParams.RoleToActorMap.Contains(AnimTrack.Role))
+			{
+				if (AActor* PreviewActor = SpawnPreviewActor(AnimTrack))
 				{
-					if (AActor* PreviewActor = SpawnPreviewActor(AnimTrack))
-					{
-						StartSceneParams.RoleToActorMap.Add(AnimTrack.Role, PreviewActor);
-					}
+					StartSceneParams.RoleToActorMap.Add(AnimTrack.Role, PreviewActor);
 				}
-
-			} // End AnimTracks Loop
-		} // End AnimSets Loop
-	} // End Sections Loop
+			}
+		}
+	}
 
 	Sequencer->NotifyMovieSceneDataChanged(EMovieSceneDataChangeType::MovieSceneStructureItemsChanged);
 
@@ -391,8 +390,10 @@ void FContextualAnimViewModel::ToggleSimulateMode()
 
 void FContextualAnimViewModel::StartSimulation()
 {
-	//SceneInstance = ContextualAnimManager->ForceStartScene(*GetSceneAsset(), StartSceneParams);
-	SceneInstance = ContextualAnimManager->TryStartScene(*GetSceneAsset(), StartSceneParams);
+	FContextualAnimStartSceneParams Params;
+	Params.RoleToActorMap = StartSceneParams.RoleToActorMap;
+	Params.SectionIdx = 0;
+	SceneInstance = ContextualAnimManager->TryStartScene(*GetSceneAsset(), Params);
 
 	if(SceneInstance == nullptr)
 	{
@@ -461,76 +462,39 @@ void FContextualAnimViewModel::SequencerTimeChanged()
 	const float CurrentSequencerTime = Sequencer->GetGlobalTime().AsSeconds();
 	const float PlaybackSpeed = Sequencer->GetPlaybackSpeed();
 
-	if(SceneInstance.IsValid())
+	if (SceneInstance.IsValid())
 	{
 		for (auto& Binding : SceneInstance->GetBindings())
 		{
-			FAnimMontageInstance* MontageInstance = Binding.GetAnimMontageInstance();
-			if (MontageInstance == nullptr)
-			{
-				continue;
-			}
-
-			const UContextualAnimMovieSceneTrack* MasterTrack = FindMasterTrackByRole(Binding.GetRoleDef().Name);
-			check(MasterTrack);
-
-			const FFrameNumber CurrentFrame = (CurrentSequencerTime * MovieSceneSequence->GetMovieScene()->GetTickResolution()).RoundToFrame();
-			const UContextualAnimMovieSceneSection* Section = Cast<const UContextualAnimMovieSceneSection>(MovieSceneHelpers::FindSectionAtTime(MasterTrack->GetAllSections(), CurrentFrame));
-
-			if (Section == nullptr || Section->GetAnimTrack().Animation == nullptr)
-			{
-				continue;
-			}
-
-			const UAnimSequenceBase* FocusedAnim = Section->GetAnimTrack().Animation;
-			const float AnimPlayLength = FocusedAnim->GetPlayLength();
-			const float SectionStartTime = (Section->GetInclusiveStartFrame() / MovieSceneSequence->GetMovieScene()->GetTickResolution());
-			const float PrevLocalTime = FMath::Clamp(PreviousSequencerTime - SectionStartTime, 0.f, AnimPlayLength);
-			const float CurrLocalTime = FMath::Clamp(CurrentSequencerTime - SectionStartTime, 0.f, AnimPlayLength);
-
-			bool bShouldTransition = false;
-			if(const UAnimMontage* FocusedAnimAsMontage = Cast<const UAnimMontage>(FocusedAnim))
-			{
-				bShouldTransition = FocusedAnimAsMontage != MontageInstance->Montage;
-			}
-			else
-			{
-				bShouldTransition = FocusedAnim != MontageInstance->Montage->SlotAnimTracks[0].AnimTrack.AnimSegments[0].GetAnimReference();
-			}
-
-			if (bShouldTransition)
+			if (FAnimMontageInstance* MontageInstance = Binding.GetAnimMontageInstance())
 			{
-				SceneInstance->TransitionTo(Binding, Section->GetAnimTrack());
-				MontageInstance = Binding.GetAnimMontageInstance();
-				check(MontageInstance);
+				const float AnimPlayLength = MontageInstance->Montage->GetPlayLength();
+				float PreviousTime = FMath::Clamp(PreviousSequencerTime, 0.f, AnimPlayLength);
+				float CurrentTime = FMath::Clamp(CurrentSequencerTime, 0.f, AnimPlayLength);
 
-				MontageInstance->bEnableAutoBlendOut = false;
-				MontageInstance->SetPosition(CurrLocalTime);
-				UpdatePreviewActorTransform(Binding, CurrLocalTime);
-			}
+				if (CurrentStatus == EMovieScenePlayerStatus::Stopped || CurrentStatus == EMovieScenePlayerStatus::Scrubbing)
+				{
+					UpdatePreviewActorTransform(Binding, CurrentTime);
 
-			if (CurrentStatus == EMovieScenePlayerStatus::Stopped || CurrentStatus == EMovieScenePlayerStatus::Scrubbing)
-			{
-				UpdatePreviewActorTransform(Binding, CurrLocalTime);
+					if (MontageInstance->IsPlaying())
+					{
+						MontageInstance->Pause();
+					}
 
-				if (MontageInstance->IsPlaying())
-				{
-					MontageInstance->Pause();
+					MontageInstance->SetPosition(CurrentTime);
 				}
-
-				MontageInstance->SetPosition(CurrLocalTime);
-			}
-			else if (CurrentStatus == EMovieScenePlayerStatus::Playing)
-			{
-				if (PlaybackSpeed > 0.f && CurrLocalTime < PrevLocalTime)
+				else if (CurrentStatus == EMovieScenePlayerStatus::Playing)
 				{
-					UpdatePreviewActorTransform(Binding, CurrLocalTime);
-					MontageInstance->SetPosition(CurrLocalTime);
-				}
+					if (PlaybackSpeed > 0.f && CurrentTime < PreviousTime)
+					{
+						UpdatePreviewActorTransform(Binding, CurrentTime);
+						MontageInstance->SetPosition(CurrentTime);
+					}
 
-				if (!MontageInstance->IsPlaying())
-				{
-					MontageInstance->SetPlaying(true);
+					if (!MontageInstance->IsPlaying())
+					{
+						MontageInstance->SetPlaying(true);
+					}
 				}
 			}
 		}
diff --git a/Engine/Plugins/Experimental/Animation/ContextualAnimation/Source/ContextualAnimationEditor/Public/ContextualAnimAssetEditorToolkit.h b/Engine/Plugins/Experimental/Animation/ContextualAnimation/Source/ContextualAnimationEditor/Public/ContextualAnimAssetEditorToolkit.h
--- a/Engine/Plugins/Experimental/Animation/ContextualAnimation/Source/ContextualAnimationEditor/Public/ContextualAnimAssetEditorToolkit.h
+++ b/Engine/Plugins/Experimental/Animation/ContextualAnimation/Source/ContextualAnimationEditor/Public/ContextualAnimAssetEditorToolkit.h
@@ -52,6 +52,7 @@
 	void ExtendToolbar();
 	void FillToolbar(FToolBarBuilder& ToolbarBuilder);
 
+	TSharedRef<SWidget> BuildSectionsMenu();
 	TSharedRef<SWidget> BuildNewAnimSetWidget();
 
 	void OnFinishedChangingProperties(const FPropertyChangedEvent& PropertyChangedEvent);
diff --git a/Engine/Plugins/Experimental/Animation/ContextualAnimation/Source/ContextualAnimationEditor/Public/ContextualAnimViewModel.h b/Engine/Plugins/Experimental/Animation/ContextualAnimation/Source/ContextualAnimationEditor/Public/ContextualAnimViewModel.h
--- a/Engine/Plugins/Experimental/Animation/ContextualAnimation/Source/ContextualAnimationEditor/Public/ContextualAnimViewModel.h
+++ b/Engine/Plugins/Experimental/Animation/ContextualAnimation/Source/ContextualAnimationEditor/Public/ContextualAnimViewModel.h
@@ -58,6 +58,8 @@
 
 	void AnimationModified(UAnimSequenceBase& Animation);
 
+	void SetActiveSection(int32 SectionIdx);
+
 	void SetActiveAnimSetForSection(int32 SectionIdx, int32 AnimSetIdx);
 
 	void OnPreviewActorClassChanged();
@@ -89,6 +91,9 @@
 
 	TWeakObjectPtr<UContextualAnimSceneInstance> SceneInstance;
 
+	/** Active section idx */
+	int32 ActiveSectionIdx = 0;
+
 	/** Active anim set for each section. Key = SectionIdx, Value = AnimSetIdx  */
 	TMap<int32, int32> ActiveAnimSetMap;
 
