[ContextualAnimSceneInstance]
- added IsDonePlaying for systems requiring polling (e.g. StateTree)
- OnSectionEndTimeReached delegate is now blueprint assignable
- added ForceTransitionToSection to the API

diff --git a/Engine/Plugins/Experimental/Animation/ContextualAnimation/Source/ContextualAnimation/Private/ContextualAnimSceneInstance.cpp b/Engine/Plugins/Experimental/Animation/ContextualAnimation/Source/ContextualAnimation/Private/ContextualAnimSceneInstance.cpp
--- a/Engine/Plugins/Experimental/Animation/ContextualAnimation/Source/ContextualAnimation/Private/ContextualAnimSceneInstance.cpp
+++ b/Engine/Plugins/Experimental/Animation/ContextualAnimation/Source/ContextualAnimation/Private/ContextualAnimSceneInstance.cpp
@@ -30,8 +30,8 @@ void UContextualAnimSceneInstance::Tick(const float DeltaTime)
 	RemainingDuration -= DeltaTime;
 	if (RemainingDuration <= 0.f)
 	{
-		OnSectionEndTimeReached.Broadcast(this);
 		RemainingDuration = MAX_flt;
+		OnSectionEndTimeReached.Broadcast(this);
 	}
 }
 
@@ -110,15 +110,25 @@ float UContextualAnimSceneInstance::Join(FContextualAnimSceneBinding& Binding)
 		}
 	}
 
-	if (UMotionWarpingComponent* MotionWarpComp = Actor->FindComponentByClass<UMotionWarpingComponent>())
+	UMotionWarpingComponent* MotionWarpComp = Actor->FindComponentByClass<UMotionWarpingComponent>();
+	for (int32 PivotIndex = 0; PivotIndex < AlignmentSectionToScenePivotList.Num(); PivotIndex++)
 	{
-		for (const FContextualAnimSetPivot& Pivot : AlignmentSectionToScenePivotList)
+		const FContextualAnimSetPivot& Pivot = AlignmentSectionToScenePivotList[PivotIndex];
+		const float Time = Binding.GetAnimTrack().GetSyncTimeForWarpSection(Pivot.Name);
+
+		if (MotionWarpComp)
 		{
-			const float Time = Binding.GetAnimTrack().GetSyncTimeForWarpSection(Pivot.Name);
 			const FTransform TransformRelativeToScenePivot = Binding.GetAnimTrack().AlignmentData.ExtractTransformAtTime(Pivot.Name, Time);
 			const FTransform WarpTarget = TransformRelativeToScenePivot * Pivot.Transform;
 			MotionWarpComp->AddOrUpdateWarpTargetFromTransform(Pivot.Name, WarpTarget);
 		}
+		else if (PivotIndex == 0)
+		{
+			// In case motion warping is not available, we use the first alignment section at time T=0 to teleport the actor
+			const FTransform TransformRelativeToScenePivot = Binding.GetAnimTrack().AlignmentData.ExtractTransformAtTime(Pivot.Name, 0.f);
+			const FTransform ActorTransform = TransformRelativeToScenePivot * Pivot.Transform;
+			Actor->TeleportTo(ActorTransform.GetLocation(), ActorTransform.GetRotation().Rotator(), /*bIsATest*/false, /*bNoCheck*/true);
+		}
 	}
 
 	if (SceneAsset->GetDisableCollisionBetweenActors())
@@ -149,27 +159,42 @@ void UContextualAnimSceneInstance::Leave(FContextualAnimSceneBinding& Binding)
 	}
 }
 
-bool UContextualAnimSceneInstance::TransitionTo(FContextualAnimSceneBinding& Binding, const FContextualAnimTrack& AnimTrack)
+float UContextualAnimSceneInstance::TransitionTo(FContextualAnimSceneBinding& Binding, const FContextualAnimTrack& AnimTrack)
 {
+	float Duration = MIN_flt;
 	check(AnimTrack.Animation != Binding.GetAnimTrack().Animation);
 	check(AnimTrack.Role == Binding.GetRoleDef().Name);
 
 	UAnimInstance* AnimInstance = Binding.GetAnimInstance();
 	if (AnimInstance == nullptr)
 	{		
-		return false;
+		return Duration;
 	}
 
 	// Unbind blend out delegate for a moment so we don't get it during the transition
 	// @TODO: Replace this with the TGuardValue 'pattern', similar to what we do in the editor for OnAnimNotifyChanged
 	AnimInstance->OnMontageBlendingOut.RemoveDynamic(this, &UContextualAnimSceneInstance::OnMontageBlendingOut);
 
-	PlayAnimation(*AnimInstance, *AnimTrack.Animation);
+	const UAnimMontage* Montage = PlayAnimation(*AnimInstance, *AnimTrack.Animation);
 	Binding.AnimTrackPtr = &AnimTrack;
 
 	AnimInstance->OnMontageBlendingOut.AddUniqueDynamic(this, &UContextualAnimSceneInstance::OnMontageBlendingOut);
 
-	return true;
+	if (Montage != nullptr)
+	{
+		const float AdjustedPlayRate = AnimInstance->Montage_GetPlayRate(Montage) * Montage->RateScale;				
+		if (AdjustedPlayRate > 0.f)
+		{
+			Duration = (Montage->GetPlayLength() / AdjustedPlayRate);	
+		}
+		else
+		{
+			UE_LOG(LogContextualAnim, Warning, TEXT("Undesired playrate %.3f, using montage play length instead."), AdjustedPlayRate);
+			Duration = Montage->GetPlayLength();
+		}
+	}
+
+	return Duration;
 }
 
 void UContextualAnimSceneInstance::Start()
@@ -191,6 +216,64 @@ void UContextualAnimSceneInstance::Stop()
 	}
 }
 
+bool UContextualAnimSceneInstance::IsDonePlaying() const
+{
+	return RemainingDuration == MAX_flt;
+}
+
+bool UContextualAnimSceneInstance::ForceTransitionToSection(const int32 SectionIdx, const int32 AnimSetIdx, const TArray<FContextualAnimSetPivot>& Pivots)
+{
+	if (!IsValid(SceneAsset))
+	{
+		UE_LOG(LogContextualAnim, Error, TEXT("Failed transition to section '%d': invalid scene asset"), SectionIdx);
+		return false;
+	}
+
+	const int32 NumSections = SceneAsset->GetNumSections();
+	if (NumSections == 0)
+	{
+		UE_LOG(LogContextualAnim, Error, TEXT("Failed transition to section '%d': no sections defined in asset %s"), SectionIdx, *SceneAsset->GetName());
+		return false;
+	}
+
+	if (SectionIdx < 0 || SectionIdx >= NumSections)
+	{
+		UE_LOG(LogContextualAnim, Error, TEXT("Failed transition to section '%d': index not in valid range [0, %d] for asset %s"), SectionIdx, NumSections, *SceneAsset->GetName());
+		return false;
+	}
+
+	if (SectionIdx == Bindings.GetSectionIdx())
+	{
+		UE_LOG(LogContextualAnim, Error, TEXT("Failed transition to section '%d': section is already playing"), SectionIdx);
+		return false;
+	}
+
+	const FName PrimaryRole = SceneAsset->GetPrimaryRole();
+	const FContextualAnimSceneBinding* PrimaryRoleBinding = Bindings.FindBindingByRole(PrimaryRole);
+	if (PrimaryRoleBinding == nullptr)
+	{
+		UE_LOG(LogContextualAnim, Error, TEXT("Failed transition to section '%d': unable to find scene binding for primary role: %s"), SectionIdx, *PrimaryRole.ToString());
+		return false;
+	}
+
+	SetPivots(Pivots);
+	
+	RemainingDuration = 0.f;
+	for (FContextualAnimSceneBinding& Binding : Bindings)
+	{
+		const FContextualAnimTrack* AnimTrack = SceneAsset->GetAnimTrack(SectionIdx, AnimSetIdx, Binding.GetRoleDef().Name);
+		if (AnimTrack == nullptr)
+		{
+			return false;
+		}
+
+		const float TrackDuration = TransitionTo(Binding, *AnimTrack);
+		RemainingDuration = FMath::Max(RemainingDuration, TrackDuration);
+	}
+
+	return true;
+}
+
 void UContextualAnimSceneInstance::OnMontageBlendingOut(UAnimMontage* Montage, bool bInterrupted)
 {
 	UE_LOG(LogContextualAnim, Log, TEXT("UContextualAnimSceneInstance::OnMontageBlendingOut Montage: %s"), *GetNameSafe(Montage));
@@ -269,7 +352,13 @@ void UContextualAnimSceneInstance::OnMontageBlendingOut(UAnimMontage* Montage, b
 	if (bShouldEnd)
 	{
 		OnSceneEnded.Broadcast(this);
-		OnSectionDonePlaying.Broadcast(this);
+
+		// Won't be ticked anymore so make sure to notify OnSectionEndTimeReached if necessary
+		if (RemainingDuration != MAX_flt)
+		{
+			RemainingDuration = MAX_flt;
+			OnSectionEndTimeReached.Broadcast(this);
+		}
 	}
 }
 
@@ -308,4 +397,4 @@ void UContextualAnimSceneInstance::SetIgnoreCollisionWithOtherActors(AActor* Act
 			}
 		}
 	}
-}
\ No newline at end of file
+}
diff --git a/Engine/Plugins/Experimental/Animation/ContextualAnimation/Source/ContextualAnimation/Public/ContextualAnimSceneInstance.h b/Engine/Plugins/Experimental/Animation/ContextualAnimation/Source/ContextualAnimation/Public/ContextualAnimSceneInstance.h
--- a/Engine/Plugins/Experimental/Animation/ContextualAnimation/Source/ContextualAnimation/Public/ContextualAnimSceneInstance.h
+++ b/Engine/Plugins/Experimental/Animation/ContextualAnimation/Source/ContextualAnimation/Public/ContextualAnimSceneInstance.h
@@ -42,21 +42,17 @@ public:
 	friend class UContextualAnimManager;
 	friend class FContextualAnimViewModel;
 
-	DECLARE_MULTICAST_DELEGATE_OneParam(FGenericSceneMulticastDelegate, UContextualAnimSceneInstance*);
-
-	/**
-	 * Delegate to notify after all montages played by the scene section blended out.
-	 * Will not be broadcasted if one or more montages have 'bEnableAutoBlendOut' set to 'false'.
-	 */
-	FGenericSceneMulticastDelegate OnSectionDonePlaying;
-
 	/**
 	 * Delegate to notify once the scene play time reaches the duration defined by the longest played montage of the selected section.
 	 * This delegate should be used if one or more montages have 'bEnableAutoBlendOut' set to 'false'.
 	 */
-	FGenericSceneMulticastDelegate OnSectionEndTimeReached;
+	UPROPERTY(BlueprintAssignable)
+	FOnContextualAnimSceneEnded OnSectionEndTimeReached;
 
-	/** Delegate to notify external objects when this is scene is completed */
+	/**
+	 * Delegate to notify external objects when this is scene is completed after all montages played by the scene section blended out.
+	 * Will not be broadcasted if one or more montages have 'bEnableAutoBlendOut' set to 'false'.
+	 */
 	UPROPERTY(BlueprintAssignable)
 	FOnContextualAnimSceneEnded OnSceneEnded;
 
@@ -85,9 +81,18 @@ public:
 	/** Resolve initial alignment and start playing animation for all actors */
 	void Start();
 
+	/**
+	 * Tells current scene instance to transition to a different section.
+	 * @note The method assumes that selection criteria were applied through bindings creation before calling.
+	 * @return True if scene was able to transition all bindings to the new section, false otherwise.
+	 */
+	bool ForceTransitionToSection(const int32 SectionIdx, const int32 AnimSetIdx, const TArray<FContextualAnimSetPivot>& Pivots);
+		
 	/** Force all the actors to leave the scene */
 	void Stop();
 
+	bool IsDonePlaying() const;
+
 	/** Whether the supplied actor is part of this scene */
 	bool IsActorInThisScene(const AActor* Actor) const;
 
@@ -111,7 +116,7 @@ protected:
 	/** Tells the scene actor to leave the scene (stop animation) */
 	void Leave(FContextualAnimSceneBinding& Binding);
 
-	bool TransitionTo(FContextualAnimSceneBinding& Binding, const FContextualAnimTrack& AnimTrack);
+	float TransitionTo(FContextualAnimSceneBinding& Binding, const FContextualAnimTrack& AnimTrack);
 
 	TArray<FContextualAnimSetPivot>& GetMutablePivots() { return AlignmentSectionToScenePivotList; }
 	void SetPivots(const TArray<FContextualAnimSetPivot>& Pivots) { AlignmentSectionToScenePivotList = Pivots; }
