ContextualAnimEditor:
 - Removed custom preview class and the option to control the actor inside the preview scene in simulate mode.
 - Added option to drag actor around in the preview scene. In simulate mode this allow the user to see the interaction start from different position. In normal mode this allow the user to reposition the actor in the scene.

diff --git a/Engine/Plugins/Experimental/Animation/ContextualAnimation/Source/ContextualAnimation/Public/ContextualAnimTypes.h b/Engine/Plugins/Experimental/Animation/ContextualAnimation/Source/ContextualAnimation/Public/ContextualAnimTypes.h
--- a/Engine/Plugins/Experimental/Animation/ContextualAnimation/Source/ContextualAnimation/Public/ContextualAnimTypes.h
+++ b/Engine/Plugins/Experimental/Animation/ContextualAnimation/Source/ContextualAnimation/Public/ContextualAnimTypes.h
@@ -96,6 +96,9 @@ struct CONTEXTUALANIMATION_API FContextualAnimTrack
 	UPROPERTY(VisibleAnywhere, BlueprintReadOnly, Category = "Defaults")
 	int32 AnimSetIdx = INDEX_NONE;
 
+	UPROPERTY(VisibleAnywhere, BlueprintReadOnly, Category = "Defaults")
+	int32 AnimTrackIdx = INDEX_NONE;
+
 	float GetSyncTimeForWarpSection(int32 WarpSectionIndex) const;
 	float GetSyncTimeForWarpSection(const FName& WarpSectionName) const;
 
@@ -209,7 +212,7 @@ struct FContextualAnimRoleDefinition
 	FName Name;
 
 	UPROPERTY(EditAnywhere, BlueprintReadWrite, Category = "Defaults")
-	TSubclassOf<AActor> PreviewActorClass;
+	bool bIsCharacter = false;
 
 	UPROPERTY(EditAnywhere, BlueprintReadWrite, Category = "Defaults")
 	FTransform MeshToComponent = FTransform(FRotator(0.f, -90.f, 0.f));
diff --git a/Engine/Plugins/Experimental/Animation/ContextualAnimation/Source/ContextualAnimationEditor/ContextualAnimationEditor.Build.cs b/Engine/Plugins/Experimental/Animation/ContextualAnimation/Source/ContextualAnimationEditor/ContextualAnimationEditor.Build.cs
--- a/Engine/Plugins/Experimental/Animation/ContextualAnimation/Source/ContextualAnimationEditor/ContextualAnimationEditor.Build.cs
+++ b/Engine/Plugins/Experimental/Animation/ContextualAnimation/Source/ContextualAnimationEditor/ContextualAnimationEditor.Build.cs
@@ -61,7 +61,8 @@ public class ContextualAnimationEditor : ModuleRules
 				"ClassViewer",
 				"Persona",
 				"RenderCore",
-				"BlueprintGraph"
+				"BlueprintGraph",
+				"AnimGraph"
 			}
 			);
 		
diff --git a/Engine/Plugins/Experimental/Animation/ContextualAnimation/Source/ContextualAnimationEditor/Private/ContextualAnimEdMode.cpp b/Engine/Plugins/Experimental/Animation/ContextualAnimation/Source/ContextualAnimationEditor/Private/ContextualAnimEdMode.cpp
--- a/Engine/Plugins/Experimental/Animation/ContextualAnimation/Source/ContextualAnimationEditor/Private/ContextualAnimEdMode.cpp
+++ b/Engine/Plugins/Experimental/Animation/ContextualAnimation/Source/ContextualAnimationEditor/Private/ContextualAnimEdMode.cpp
@@ -1,406 +1,287 @@
-// Copyright Epic Games, Inc. All Rights Reserved.
-
-#include "ContextualAnimEdMode.h"
-#include "EngineUtils.h"
-#include "ContextualAnimSceneAsset.h"
-#include "ContextualAnimSelectionCriterion.h"
-#include "ContextualAnimUtilities.h"
-#include "ContextualAnimSceneInstance.h"
-#include "ContextualAnimViewportClient.h"
-#include "ContextualAnimAssetEditorToolkit.h"
-#include "ContextualAnimViewModel.h"
-#include "Animation/AnimMontage.h"
-#include "AnimNotifyState_IKWindow.h"
-#include "SkeletalDebugRendering.h"
-
-#include "GameFramework/Character.h"
-#include "Components/CapsuleComponent.h"
-#include "Components/SkeletalMeshComponent.h"
-
-const FEditorModeID FContextualAnimEdMode::EdModeId = TEXT("ContextualAnimEdMode");
-
-IMPLEMENT_HIT_PROXY(HSelectionCriterionHitProxy, HHitProxy);
-
-FContextualAnimEdMode::FContextualAnimEdMode()
-{
-}
-
-FContextualAnimEdMode::~FContextualAnimEdMode()
-{
-}
-
-void FContextualAnimEdMode::Tick(FEditorViewportClient* ViewportClient, float DeltaTime)
-{
-	FEdMode::Tick(ViewportClient, DeltaTime);
-
-	if (ViewModel && ViewModel->IsSimulateModeActive())
-	{
-		// @TODO: Two things here are temporal. First, we have to click on the actor that we want to control, and second, we can only control characters. 
-		// We need to find a better work flow for this, but its good enough for now as a prototype
-		if (!ViewportClient->Viewport->KeyState(EKeys::RightMouseButton))
-		{
-			if (ViewportClient->Viewport->KeyState(EKeys::W) || ViewportClient->Viewport->KeyState(EKeys::S))
-			{
-				if (ACharacter* Character = Cast<ACharacter>(SelectedActor.Get()))
-				{
-					const FVector WorldDirection = FRotationMatrix(Character->GetActorRotation()).GetScaledAxis(EAxis::X);
-					Character->AddMovementInput(WorldDirection, ViewportClient->Viewport->KeyState(EKeys::W) ? 1.f : -1.f);
-				}
-			}
-
-			if (ViewportClient->Viewport->KeyState(EKeys::A) || ViewportClient->Viewport->KeyState(EKeys::D))
-			{
-				if (ACharacter* Character = Cast<ACharacter>(SelectedActor.Get()))
-				{
-					const FVector WorldDirection = FRotationMatrix(Character->GetActorRotation()).GetScaledAxis(EAxis::Y);
-					Character->AddMovementInput(WorldDirection, ViewportClient->Viewport->KeyState(EKeys::D) ? 1.f : -1.f);
-				}
-			}
-		}
-	}
-}
-
-void FContextualAnimEdMode::Render(const FSceneView* View, FViewport* Viewport, FPrimitiveDrawInterface* PDI)
-{
-	FEdMode::Render(View, Viewport, PDI);
-
-	// @TODO: This should not be initialized here
-	FContextualAnimViewportClient* ViewportClient = static_cast<FContextualAnimViewportClient*>(Viewport->GetClient());
-	if(!ViewModel)
-	{
-		ViewModel = ViewportClient->GetAssetEditorToolkit()->GetViewModel();
-	}
-
-	if (ViewModel)
-	{
-		// Draw a circle below the selected actor if it is a character to let the user know that it is a controllable actor during simulate mode
-		if(ViewModel->IsSimulateModeActive() && SelectedActor.IsValid())
-		{
-			if (const ACharacter* Character = Cast<ACharacter>(SelectedActor.Get()))
-			{
-				const float CapsuleHalfHeight = Character->GetCapsuleComponent()->GetScaledCapsuleHalfHeight();
-				const FVector Location = Character->GetActorLocation() - Character->GetActorQuat().GetUpVector() * CapsuleHalfHeight;
-				DrawCircle(PDI, Location, FVector(1, 0, 0), FVector(0, 1, 0), FColor::Red, Character->GetCapsuleComponent()->GetScaledCapsuleRadius(), 12, SDPG_World, 1.f);
-			}
-		}
-
-		if (const UContextualAnimSceneInstance* SceneInstance = ViewModel->GetSceneInstance())
-		{
-			const UContextualAnimSceneAsset& SceneAsset = SceneInstance->GetSceneAsset();
-			const FContextualAnimSceneBindings& Bindings = SceneInstance->GetBindings();
-
-			// Draw Scene Pivots
-			if(const FContextualAnimSceneSection* Section = SceneAsset.GetSection(Bindings.GetSectionIdx()))
-			{
-				if(const FContextualAnimSet* AnimSet = Section->GetAnimSet(Bindings.GetAnimSetIdx()))
-				{
-					for (const FTransform& ScenePivot : AnimSet->ScenePivots)
-					{
-						DrawCoordinateSystem(PDI, ScenePivot.GetLocation(), ScenePivot.Rotator(), 50.f, SDPG_Foreground);
-					}
-				}
-			}
-
-			const EShowIKTargetsDrawMode IKTargetsDrawMode = ViewportClient->GetShowIKTargetsDrawMode();
-
-			for (const FContextualAnimSceneBinding& Binding : Bindings)
-			{
-				// Draw IK Targets
-				if (IKTargetsDrawMode == EShowIKTargetsDrawMode::All || (IKTargetsDrawMode == EShowIKTargetsDrawMode::Selected && Binding.GetActor() == SelectedActor.Get()))
-				{
-					DrawIKTargetsForBinding(*PDI, Binding);
-				}
-
-				// Draw Selection Criteria
-				if(const FContextualAnimSceneBinding* PrimaryBinding = Bindings.FindBindingByRole(SceneAsset.GetPrimaryRole()))
-				{
-					const FTransform PrimaryTransform = PrimaryBinding->GetTransform();
-
-					const FContextualAnimTrack& AnimTrack = Binding.GetAnimTrack();
-					for (int32 CriterionIdx = 0; CriterionIdx < AnimTrack.SelectionCriteria.Num(); CriterionIdx++)
-					{
-						if(const UContextualAnimSelectionCriterion* Criterion = AnimTrack.SelectionCriteria[CriterionIdx])
-						{
-							FLinearColor DrawColor = FLinearColor::White;
-							if (Criterion->DoesQuerierPassCondition(FContextualAnimSceneBindingContext(PrimaryTransform), Binding.GetContext()))
-							{
-								DrawColor = FLinearColor::Green;
-							}
-
-							//@TODO: Each SelectionCriterion should implement this, and here we should just call "Criterion->Draw()"
-							if (const UContextualAnimSelectionCriterion_TriggerArea* Spatial = Cast<UContextualAnimSelectionCriterion_TriggerArea>(Criterion))
-							{
-								const float HalfHeight = Spatial->Height / 2.f;
-								const int32 LastIndex = Spatial->PolygonPoints.Num() - 1;
-								for (int32 Idx = 0; Idx <= LastIndex; Idx++)
-								{
-									const FVector P0 = PrimaryTransform.TransformPositionNoScale(Spatial->PolygonPoints[Idx]);
-									const FVector P1 = PrimaryTransform.TransformPositionNoScale(Spatial->PolygonPoints[Idx == LastIndex ? 0 : Idx + 1]);
-
-									PDI->DrawLine(P0, P1, DrawColor, SDPG_Foreground, 2.f);
-
-									PDI->DrawLine(P0 + FVector::UpVector * Spatial->Height, P1 + FVector::UpVector * Spatial->Height, DrawColor, SDPG_Foreground, 2.f);
-
-									PDI->DrawLine(P0, P0 + FVector::UpVector * Spatial->Height, DrawColor, SDPG_Foreground, 2.f);
-
-									PDI->SetHitProxy(new HSelectionCriterionHitProxy(FSelectionCriterionHitProxyData(AnimTrack.SectionIdx, AnimTrack.AnimSetIdx, AnimTrack.Role, CriterionIdx, Idx)));
-									PDI->DrawPoint(P0, FLinearColor::Black, 15.f, SDPG_Foreground);
-									PDI->SetHitProxy(nullptr);
-
-									PDI->SetHitProxy(new HSelectionCriterionHitProxy(FSelectionCriterionHitProxyData(AnimTrack.SectionIdx, AnimTrack.AnimSetIdx, AnimTrack.Role, CriterionIdx, Idx + 4)));
-									PDI->DrawPoint(P0 + FVector::UpVector * Spatial->Height, FLinearColor::Black, 15.f, SDPG_Foreground);
-									PDI->SetHitProxy(nullptr);
-								}
-							}
-							else if (const UContextualAnimSelectionCriterion_Facing* Facing = Cast<UContextualAnimSelectionCriterion_Facing>(Criterion))
-							{
-								const FTransform Transform = AnimTrack.GetAlignmentTransformAtEntryTime() * PrimaryTransform;
-								UContextualAnimUtilities::DrawSector(*PDI, Transform.GetLocation(), Transform.GetRotation().GetForwardVector(), 0.f, 30.f, -Facing->MaxAngle, Facing->MaxAngle, DrawColor, SDPG_World, 1.f);
-							}
-						}
-					}
-				}
-			}
-		}
-	}
-}
-
-void FContextualAnimEdMode::DrawIKTargetsForBinding(FPrimitiveDrawInterface& PDI, const FContextualAnimSceneBinding& Binding) const
-{
-	for (const FContextualAnimIKTargetDefinition& IKTargetDef : Binding.GetIKTargetDefs().IKTargetDefs)
-	{
-		if (const FContextualAnimSceneBinding* TargetBinding = Binding.GetSceneInstance()->FindBindingByRole(IKTargetDef.TargetRoleName))
-		{
-			if (USkeletalMeshComponent* TargetSkelMeshComp = TargetBinding->GetSkeletalMeshComponent())
-			{
-				const float Alpha = UAnimNotifyState_IKWindow::GetIKAlphaValue(IKTargetDef.GoalName, Binding.GetAnimMontageInstance());
-
-				if (IKTargetDef.Provider == EContextualAnimIKTargetProvider::Bone)
-				{
-					//@TODO:...
-				}
-				else if (IKTargetDef.Provider == EContextualAnimIKTargetProvider::Autogenerated)
-				{
-					const FTransform ParentTransform = TargetSkelMeshComp->GetSocketTransform(IKTargetDef.TargetBoneName);
-
-					const FTransform TargetTransform = Binding.GetAnimTrack().IKTargetData.ExtractTransformAtTime(IKTargetDef.GoalName, Binding.GetAnimMontageTime()) * ParentTransform;
-
-					FLinearColor Color = Alpha > 0.f ? FLinearColor(FColor::MakeRedToGreenColorFromScalar(Alpha)) : FLinearColor::White;
-
-					FVector Start = ParentTransform.GetLocation();
-					FVector End = TargetTransform.GetLocation();
-
-					const float Radius = 1.f;
-					SkeletalDebugRendering::DrawWireBone(&PDI, Start, End, Color, SDPG_Foreground, Radius);
-					SkeletalDebugRendering::DrawAxes(&PDI, FTransform(End), SDPG_Foreground, 0.f, Radius);
-				}
-			}
-		}
-	}
-}
-
-bool FContextualAnimEdMode::HandleClick(FEditorViewportClient* InViewportClient, HHitProxy* HitProxy, const FViewportClick& Click)
-{
-	SelectedActor.Reset();
-	SelectedSelectionCriterionData.Reset();
-
-	if (HitProxy != nullptr)
-	{
-		if (HitProxy->IsA(HActor::StaticGetType()))
-		{
-			HActor* ActorHitProxy = static_cast<HActor*>(HitProxy);
-			SelectedActor = ActorHitProxy->Actor;
-			return true;
-		}
-		else if (HitProxy->IsA(HSelectionCriterionHitProxy::StaticGetType()))
-		{
-			SelectedSelectionCriterionData = static_cast<HSelectionCriterionHitProxy*>(HitProxy)->Data;
-			return true;
-		}
-	}
-
-	return false; // unhandled
-}
-
-bool FContextualAnimEdMode::GetHitResultUnderCursor(FHitResult& OutHitResult, FEditorViewportClient* InViewportClient, const FViewportClick& Click) const
-{
-	FSceneViewFamilyContext ViewFamily(FSceneViewFamily::ConstructionValues(InViewportClient->Viewport, InViewportClient->GetScene(), InViewportClient->EngineShowFlags).SetRealtimeUpdate(InViewportClient->IsRealtime()));
-	FSceneView* View = InViewportClient->CalcSceneView(&ViewFamily);
-	FViewportCursorLocation Cursor(View, InViewportClient, Click.GetClickPos().X, Click.GetClickPos().Y);
-	const auto ViewportType = InViewportClient->GetViewportType();
-
-	const FVector RayStart = Cursor.GetOrigin();
-	const FVector RayEnd = RayStart + Cursor.GetDirection() * HALF_WORLD_MAX;
-
-	return InViewportClient->GetWorld()->LineTraceSingleByChannel(OutHitResult, RayStart, RayEnd, ECC_WorldStatic, FCollisionQueryParams::DefaultQueryParam);
-}
-
-bool FContextualAnimEdMode::StartTracking(FEditorViewportClient* InViewportClient, FViewport* InViewport)
-{
-	return FEdMode::StartTracking(InViewportClient, InViewport);
-}
-
-bool FContextualAnimEdMode::EndTracking(FEditorViewportClient* InViewportClient, FViewport* InViewport)
-{
-	return FEdMode::EndTracking(InViewportClient, InViewport);
-}
-
-bool FContextualAnimEdMode::InputDelta(FEditorViewportClient* InViewportClient, FViewport* InViewport, FVector& InDrag, FRotator& InRot, FVector& InScale)
-{
-	EAxisList::Type CurrentAxis = InViewportClient->GetCurrentWidgetAxis();
-
-	if (CurrentAxis != EAxisList::None)
-	{
-		if (ViewModel)
-		{
-			if (SelectedSelectionCriterionData.IsValid())
-			{
-				const UContextualAnimSceneAsset* SceneAsset = ViewModel->GetSceneAsset();
-				if (const FContextualAnimTrack* AnimTrack = SceneAsset->GetAnimTrack(SelectedSelectionCriterionData.SectionIdx, SelectedSelectionCriterionData.AnimSetIdx, SelectedSelectionCriterionData.RoleName))
-				{
-					if(AnimTrack->SelectionCriteria.IsValidIndex(SelectedSelectionCriterionData.CriterionIdx))
-					{
-						if (UContextualAnimSelectionCriterion_TriggerArea* Spatial = Cast<UContextualAnimSelectionCriterion_TriggerArea>(AnimTrack->SelectionCriteria[SelectedSelectionCriterionData.CriterionIdx]))
-						{
-							FMatrix WidgetCoordSystem = FMatrix::Identity;
-							GetCustomDrawingCoordinateSystem(WidgetCoordSystem, nullptr);
-
-							InDrag = WidgetCoordSystem.InverseTransformVector(InDrag);
-
-							FVector& Point = Spatial->PolygonPoints[SelectedSelectionCriterionData.DataIdx >= 4 ? SelectedSelectionCriterionData.DataIdx - 4 : SelectedSelectionCriterionData.DataIdx];
-							Point.X += InDrag.X;
-							Point.Y += InDrag.Y;
-
-							if (InDrag.Z != 0.f)
-							{
-								if (SelectedSelectionCriterionData.DataIdx < 4)
-								{
-									for (int32 Idx = 0; Idx < Spatial->PolygonPoints.Num(); Idx++)
-									{
-										Spatial->PolygonPoints[Idx].Z += InDrag.Z;
-									}
-
-									Spatial->Height = FMath::Max(Spatial->Height - InDrag.Z, 0.f);
-								}
-								else
-								{
-									Spatial->Height = FMath::Max(Spatial->Height + InDrag.Z, 0.f);
-								}
-							}
-
-							return true;
-						}
-					}
-				}
-			}
-		}
-	}
-
-	return false; // unhandled
-}
-
-bool FContextualAnimEdMode::InputKey(FEditorViewportClient* ViewportClient, FViewport* Viewport, FKey Key, EInputEvent Event)
-{
-	if(ViewModel && ViewModel->IsSimulateModeActive())
-	{
-		if (Key == EKeys::Enter && Event == IE_Released)
-		{
-			ViewModel->StartSimulation();
-			return true;
-		}
-	}
-
-	return FEdMode::InputKey(ViewportClient, Viewport, Key, Event);
-}
-
-bool FContextualAnimEdMode::AllowWidgetMove()
-{
-	return ShouldDrawWidget();
-}
-
-bool FContextualAnimEdMode::ShouldDrawWidget() const
-{
-	if (ViewModel && SelectedSelectionCriterionData.IsValid())
-	{
-		const UContextualAnimSceneAsset* SceneAsset = ViewModel->GetSceneAsset();
-		if (const FContextualAnimTrack* AnimTrack = SceneAsset->GetAnimTrack(SelectedSelectionCriterionData.SectionIdx, SelectedSelectionCriterionData.AnimSetIdx, SelectedSelectionCriterionData.RoleName))
-		{
-			const int32 Idx = SelectedSelectionCriterionData.CriterionIdx;
-			if (AnimTrack->SelectionCriteria.IsValidIndex(Idx) && AnimTrack->SelectionCriteria[Idx]->GetClass()->IsChildOf<UContextualAnimSelectionCriterion_TriggerArea>())
-			{
-				return true;
-			}
-		}
-	}
-
-	return false;
-}
-
-bool FContextualAnimEdMode::GetCustomDrawingCoordinateSystem(FMatrix& InMatrix, void* InData)
-{
-	if (ViewModel)
-	{
-		if (SelectedSelectionCriterionData.IsValid())
-		{
-			if (const UContextualAnimSceneInstance* SceneInstance = ViewModel->GetSceneInstance())
-			{
-				FTransform PrimaryActorTransform = FTransform::Identity;
-				if (const FContextualAnimSceneBinding* Binding = SceneInstance->FindBindingByRole(ViewModel->GetSceneAsset()->GetPrimaryRole()))
-				{
-					PrimaryActorTransform = Binding->GetTransform();
-					InMatrix = PrimaryActorTransform.ToMatrixNoScale().RemoveTranslation();
-					return true;
-				}
-			}
-		}
-	}
-
-	return false;
-}
-
-bool FContextualAnimEdMode::GetCustomInputCoordinateSystem(FMatrix& InMatrix, void* InData)
-{
-	return GetCustomDrawingCoordinateSystem(InMatrix, InData);
-}
-
-FVector FContextualAnimEdMode::GetWidgetLocation() const
-{
-	if (ViewModel)
-	{
-		if (SelectedSelectionCriterionData.IsValid())
-		{
-			if (const UContextualAnimSceneInstance* SceneInstance = ViewModel->GetSceneInstance())
-			{
-				const UContextualAnimSceneAsset* SceneAsset = ViewModel->GetSceneAsset();
-				if (const FContextualAnimTrack* AnimTrack = SceneAsset->GetAnimTrack(SelectedSelectionCriterionData.SectionIdx, SelectedSelectionCriterionData.AnimSetIdx, SelectedSelectionCriterionData.RoleName))
-				{
-					if (AnimTrack->SelectionCriteria.IsValidIndex(SelectedSelectionCriterionData.CriterionIdx))
-					{
-						if (const UContextualAnimSelectionCriterion_TriggerArea* Spatial = Cast<UContextualAnimSelectionCriterion_TriggerArea>(AnimTrack->SelectionCriteria[SelectedSelectionCriterionData.CriterionIdx]))
-						{
-							FVector Location = FVector::ZeroVector;
-							if (SelectedSelectionCriterionData.DataIdx < 4)
-							{
-								Location = Spatial->PolygonPoints[SelectedSelectionCriterionData.DataIdx];
-							}
-							else
-							{
-								Location = Spatial->PolygonPoints[SelectedSelectionCriterionData.DataIdx - 4] + FVector::UpVector * Spatial->Height;
-							}
-
-							FTransform PrimaryActorTransform = FTransform::Identity;
-							if (const FContextualAnimSceneBinding* Binding = SceneInstance->FindBindingByRole(SceneAsset->GetPrimaryRole()))
-							{
-								PrimaryActorTransform = Binding->GetTransform();
-							}
-
-							return PrimaryActorTransform.TransformPositionNoScale(Location);
-						}
-					}
-				}
-			}
-		}
-	}
-
-	return FVector::ZeroVector;
-}
+// Copyright Epic Games, Inc. All Rights Reserved.
+
+#include "ContextualAnimEdMode.h"
+#include "EngineUtils.h"
+#include "ContextualAnimSceneAsset.h"
+#include "ContextualAnimSelectionCriterion.h"
+#include "ContextualAnimUtilities.h"
+#include "ContextualAnimSceneInstance.h"
+#include "ContextualAnimViewportClient.h"
+#include "ContextualAnimAssetEditorToolkit.h"
+#include "ContextualAnimViewModel.h"
+#include "Animation/AnimMontage.h"
+#include "AnimNotifyState_IKWindow.h"
+#include "SkeletalDebugRendering.h"
+#include "CanvasTypes.h"
+#include "CanvasItem.h"
+#include "Components/SkeletalMeshComponent.h"
+
+const FEditorModeID FContextualAnimEdMode::EdModeId = TEXT("ContextualAnimEdMode");
+
+IMPLEMENT_HIT_PROXY(HSelectionCriterionHitProxy, HHitProxy);
+
+FContextualAnimEdMode::FContextualAnimEdMode()
+{
+}
+
+FContextualAnimEdMode::~FContextualAnimEdMode()
+{
+}
+
+void FContextualAnimEdMode::Tick(FEditorViewportClient* ViewportClient, float DeltaTime)
+{
+	FEdMode::Tick(ViewportClient, DeltaTime);
+}
+
+void FContextualAnimEdMode::Render(const FSceneView* View, FViewport* Viewport, FPrimitiveDrawInterface* PDI)
+{
+	FEdMode::Render(View, Viewport, PDI);
+
+	// @TODO: This should not be initialized here
+	FContextualAnimViewportClient* ViewportClient = static_cast<FContextualAnimViewportClient*>(Viewport->GetClient());
+	if(!ViewModel)
+	{
+		ViewModel = ViewportClient->GetAssetEditorToolkit()->GetViewModel();
+	}
+
+	if (ViewModel)
+	{
+		if (const UContextualAnimSceneInstance* SceneInstance = ViewModel->GetSceneInstance())
+		{
+			const UContextualAnimSceneAsset& SceneAsset = SceneInstance->GetSceneAsset();
+			const FContextualAnimSceneBindings& Bindings = SceneInstance->GetBindings();
+
+			// Draw Scene Pivots
+			if(const FContextualAnimSceneSection* Section = SceneAsset.GetSection(Bindings.GetSectionIdx()))
+			{
+				if(const FContextualAnimSet* AnimSet = Section->GetAnimSet(Bindings.GetAnimSetIdx()))
+				{
+					for (const FTransform& ScenePivot : AnimSet->ScenePivots)
+					{
+						DrawCoordinateSystem(PDI, ScenePivot.GetLocation(), ScenePivot.Rotator(), 50.f, SDPG_Foreground);
+					}
+				}
+			}
+
+			const EShowIKTargetsDrawMode IKTargetsDrawMode = ViewportClient->GetShowIKTargetsDrawMode();
+
+			for (const FContextualAnimSceneBinding& Binding : Bindings)
+			{
+				// Draw IK Targets
+				if (IKTargetsDrawMode == EShowIKTargetsDrawMode::All || (IKTargetsDrawMode == EShowIKTargetsDrawMode::Selected && Binding.GetActor() == ViewModel->GetSelectedActor()))
+				{
+					DrawIKTargetsForBinding(*PDI, Binding);
+				}
+
+				// Draw Selection Criteria
+				if(const FContextualAnimSceneBinding* PrimaryBinding = Bindings.FindBindingByRole(SceneAsset.GetPrimaryRole()))
+				{
+					const FTransform PrimaryTransform = PrimaryBinding->GetTransform();
+
+					const FContextualAnimTrack& AnimTrack = Binding.GetAnimTrack();
+					for (int32 CriterionIdx = 0; CriterionIdx < AnimTrack.SelectionCriteria.Num(); CriterionIdx++)
+					{
+						if(const UContextualAnimSelectionCriterion* Criterion = AnimTrack.SelectionCriteria[CriterionIdx])
+						{
+							FLinearColor DrawColor = FLinearColor::White;
+							if (Criterion->DoesQuerierPassCondition(FContextualAnimSceneBindingContext(PrimaryTransform), Binding.GetContext()))
+							{
+								DrawColor = FLinearColor::Green;
+							}
+
+							//@TODO: Each SelectionCriterion should implement this, and here we should just call "Criterion->Draw()"
+							if (const UContextualAnimSelectionCriterion_TriggerArea* Spatial = Cast<UContextualAnimSelectionCriterion_TriggerArea>(Criterion))
+							{
+								const float HalfHeight = Spatial->Height / 2.f;
+								const int32 LastIndex = Spatial->PolygonPoints.Num() - 1;
+								for (int32 Idx = 0; Idx <= LastIndex; Idx++)
+								{
+									const FVector P0 = PrimaryTransform.TransformPositionNoScale(Spatial->PolygonPoints[Idx]);
+									const FVector P1 = PrimaryTransform.TransformPositionNoScale(Spatial->PolygonPoints[Idx == LastIndex ? 0 : Idx + 1]);
+
+									PDI->DrawLine(P0, P1, DrawColor, SDPG_Foreground, 2.f);
+
+									PDI->DrawLine(P0 + FVector::UpVector * Spatial->Height, P1 + FVector::UpVector * Spatial->Height, DrawColor, SDPG_Foreground, 2.f);
+
+									PDI->DrawLine(P0, P0 + FVector::UpVector * Spatial->Height, DrawColor, SDPG_Foreground, 2.f);
+
+									PDI->SetHitProxy(new HSelectionCriterionHitProxy(Binding.GetRoleDef().Name, CriterionIdx, Idx));
+									PDI->DrawPoint(P0, FLinearColor::Black, 15.f, SDPG_Foreground);
+									PDI->SetHitProxy(nullptr);
+
+									PDI->SetHitProxy(new HSelectionCriterionHitProxy(Binding.GetRoleDef().Name, CriterionIdx, Idx + 4));
+									PDI->DrawPoint(P0 + FVector::UpVector * Spatial->Height, FLinearColor::Black, 15.f, SDPG_Foreground);
+									PDI->SetHitProxy(nullptr);
+								}
+							}
+							else if (const UContextualAnimSelectionCriterion_Facing* Facing = Cast<UContextualAnimSelectionCriterion_Facing>(Criterion))
+							{
+								const FTransform Transform = AnimTrack.GetAlignmentTransformAtEntryTime() * PrimaryTransform;
+								UContextualAnimUtilities::DrawSector(*PDI, Transform.GetLocation(), Transform.GetRotation().GetForwardVector(), 0.f, 30.f, -Facing->MaxAngle, Facing->MaxAngle, DrawColor, SDPG_World, 1.f);
+							}
+						}
+					}
+				}
+			}
+		}
+	}
+}
+
+void FContextualAnimEdMode::DrawHUD(FEditorViewportClient* ViewportClient, FViewport* Viewport, const FSceneView* View, FCanvas* Canvas)
+{
+	FEdMode::DrawHUD(ViewportClient, Viewport, View, Canvas);
+
+	if (ViewModel)
+	{
+		FCanvasTextItem TextItem(FVector2D(10.f, 40.f), ViewModel->GetSelectionDebugText(), GEngine->GetSmallFont(), FLinearColor::White);
+		Canvas->DrawItem(TextItem);
+	}
+}
+
+void FContextualAnimEdMode::DrawIKTargetsForBinding(FPrimitiveDrawInterface& PDI, const FContextualAnimSceneBinding& Binding) const
+{
+	for (const FContextualAnimIKTargetDefinition& IKTargetDef : Binding.GetIKTargetDefs().IKTargetDefs)
+	{
+		if (const FContextualAnimSceneBinding* TargetBinding = Binding.GetSceneInstance()->FindBindingByRole(IKTargetDef.TargetRoleName))
+		{
+			if (USkeletalMeshComponent* TargetSkelMeshComp = TargetBinding->GetSkeletalMeshComponent())
+			{
+				const float Alpha = UAnimNotifyState_IKWindow::GetIKAlphaValue(IKTargetDef.GoalName, Binding.GetAnimMontageInstance());
+
+				if (IKTargetDef.Provider == EContextualAnimIKTargetProvider::Bone)
+				{
+					//@TODO:...
+				}
+				else if (IKTargetDef.Provider == EContextualAnimIKTargetProvider::Autogenerated)
+				{
+					const FTransform ParentTransform = TargetSkelMeshComp->GetSocketTransform(IKTargetDef.TargetBoneName);
+
+					const FTransform TargetTransform = Binding.GetAnimTrack().IKTargetData.ExtractTransformAtTime(IKTargetDef.GoalName, Binding.GetAnimMontageTime()) * ParentTransform;
+
+					FLinearColor Color = Alpha > 0.f ? FLinearColor(FColor::MakeRedToGreenColorFromScalar(Alpha)) : FLinearColor::White;
+
+					FVector Start = ParentTransform.GetLocation();
+					FVector End = TargetTransform.GetLocation();
+
+					const float Radius = 1.f;
+					SkeletalDebugRendering::DrawWireBone(&PDI, Start, End, Color, SDPG_Foreground, Radius);
+					SkeletalDebugRendering::DrawAxes(&PDI, FTransform(End), SDPG_Foreground, 0.f, Radius);
+				}
+			}
+		}
+	}
+}
+
+bool FContextualAnimEdMode::HandleClick(FEditorViewportClient* InViewportClient, HHitProxy* HitProxy, const FViewportClick& Click)
+{
+	if (HitProxy != nullptr)
+	{
+		if (HitProxy->IsA(HActor::StaticGetType()))
+		{
+			HActor* ActorHitProxy = static_cast<HActor*>(HitProxy);
+			ViewModel->UpdateSelection(ActorHitProxy->Actor);
+			return true;
+		}
+		else if (HitProxy->IsA(HSelectionCriterionHitProxy::StaticGetType()))
+		{
+			HSelectionCriterionHitProxy* CriterionHitProxy = static_cast<HSelectionCriterionHitProxy*>(HitProxy);
+			ViewModel->UpdateSelection(CriterionHitProxy->Role, CriterionHitProxy->IndexPair.Key, CriterionHitProxy->IndexPair.Value);
+			return true;
+		}
+	}
+
+	ViewModel->ClearSelection();
+	return false; // unhandled
+}
+
+bool FContextualAnimEdMode::GetHitResultUnderCursor(FHitResult& OutHitResult, FEditorViewportClient* InViewportClient, const FViewportClick& Click) const
+{
+	FSceneViewFamilyContext ViewFamily(FSceneViewFamily::ConstructionValues(InViewportClient->Viewport, InViewportClient->GetScene(), InViewportClient->EngineShowFlags).SetRealtimeUpdate(InViewportClient->IsRealtime()));
+	FSceneView* View = InViewportClient->CalcSceneView(&ViewFamily);
+	FViewportCursorLocation Cursor(View, InViewportClient, Click.GetClickPos().X, Click.GetClickPos().Y);
+	const auto ViewportType = InViewportClient->GetViewportType();
+
+	const FVector RayStart = Cursor.GetOrigin();
+	const FVector RayEnd = RayStart + Cursor.GetDirection() * HALF_WORLD_MAX;
+
+	return InViewportClient->GetWorld()->LineTraceSingleByChannel(OutHitResult, RayStart, RayEnd, ECC_WorldStatic, FCollisionQueryParams::DefaultQueryParam);
+}
+
+bool FContextualAnimEdMode::StartTracking(FEditorViewportClient* InViewportClient, FViewport* InViewport)
+{
+	return FEdMode::StartTracking(InViewportClient, InViewport);
+}
+
+bool FContextualAnimEdMode::EndTracking(FEditorViewportClient* InViewportClient, FViewport* InViewport)
+{
+	return FEdMode::EndTracking(InViewportClient, InViewport);
+}
+
+bool FContextualAnimEdMode::InputDelta(FEditorViewportClient* InViewportClient, FViewport* InViewport, FVector& InDrag, FRotator& InRot, FVector& InScale)
+{
+	EAxisList::Type CurrentAxis = InViewportClient->GetCurrentWidgetAxis();
+
+	if (CurrentAxis != EAxisList::None)
+	{
+		if (ViewModel)
+		{
+			return ViewModel->ProcessInputDelta(InDrag, InRot, InScale);
+		}
+	}
+
+	return false; // unhandled
+}
+
+bool FContextualAnimEdMode::InputKey(FEditorViewportClient* ViewportClient, FViewport* Viewport, FKey Key, EInputEvent Event)
+{
+	if(ViewModel && ViewModel->IsSimulateModeActive())
+	{
+		if (Key == EKeys::Enter && Event == IE_Released)
+		{
+			ViewModel->StartSimulation();
+			return true;
+		}
+	}
+
+	return FEdMode::InputKey(ViewportClient, Viewport, Key, Event);
+}
+
+bool FContextualAnimEdMode::AllowWidgetMove()
+{
+	return ShouldDrawWidget();
+}
+
+bool FContextualAnimEdMode::ShouldDrawWidget() const
+{
+	if (ViewModel)
+	{
+		return ViewModel->ShouldPreviewSceneDrawWidget();
+	}
+
+	return false;
+}
+
+bool FContextualAnimEdMode::GetCustomDrawingCoordinateSystem(FMatrix& InMatrix, void* InData)
+{
+	if (ViewModel)
+	{
+		return ViewModel->GetCustomDrawingCoordinateSystem(InMatrix, InData);
+	}
+
+	return false;
+}
+
+bool FContextualAnimEdMode::GetCustomInputCoordinateSystem(FMatrix& InMatrix, void* InData)
+{
+	return GetCustomDrawingCoordinateSystem(InMatrix, InData);
+}
+
+FVector FContextualAnimEdMode::GetWidgetLocation() const
+{
+	if (ViewModel)
+	{
+		return ViewModel->GetWidgetLocationFromSelection();
+	}
+
+	return FVector::ZeroVector;
+}
diff --git a/Engine/Plugins/Experimental/Animation/ContextualAnimation/Source/ContextualAnimationEditor/Private/ContextualAnimViewModel.cpp b/Engine/Plugins/Experimental/Animation/ContextualAnimation/Source/ContextualAnimationEditor/Private/ContextualAnimViewModel.cpp
--- a/Engine/Plugins/Experimental/Animation/ContextualAnimation/Source/ContextualAnimationEditor/Private/ContextualAnimViewModel.cpp
+++ b/Engine/Plugins/Experimental/Animation/ContextualAnimation/Source/ContextualAnimationEditor/Private/ContextualAnimViewModel.cpp
@@ -30,6 +30,9 @@
 #include "Camera/CameraComponent.h"
 #include "Components/CapsuleComponent.h"
 #include "MotionWarpingComponent.h"
+#include "Animation/DebugSkelMeshComponent.h"
+#include "AnimPreviewInstance.h"
+#include "ContextualAnimSelectionCriterion.h"
 
 FContextualAnimViewModel::FContextualAnimViewModel()
 	: SceneAsset(nullptr)
@@ -127,17 +130,43 @@ void FContextualAnimViewModel::SetActiveAnimSetForSection(int32 SectionIdx, int3
 
 AActor* FContextualAnimViewModel::SpawnPreviewActor(const FContextualAnimTrack& AnimTrack)
 {
+	if (AnimTrack.Animation == nullptr)
+	{
+		return nullptr;
+	}
+
+	USkeleton* Skeleton = AnimTrack.Animation->GetSkeleton();
+
+	if (Skeleton == nullptr)
+	{
+		return nullptr;
+	}
+
 	const FContextualAnimRoleDefinition* RoleDef = GetSceneAsset()->RolesAsset ? GetSceneAsset()->RolesAsset->FindRoleDefinitionByName(AnimTrack.Role) : nullptr;
-	UClass* PreviewClass = RoleDef ? RoleDef->PreviewActorClass : nullptr;
+	const bool bIsCharacter = (RoleDef && RoleDef->bIsCharacter);
+
 	const FTransform SpawnTransform = AnimTrack.GetRootTransformAtTime(0.f);
 
 	FActorSpawnParameters Params;
 	Params.SpawnCollisionHandlingOverride = ESpawnActorCollisionHandlingMethod::AdjustIfPossibleButAlwaysSpawn;
-	AActor* PreviewActor = GetWorld()->SpawnActor<AActor>(PreviewClass, SpawnTransform, Params);
 
-	if (ACharacter* PreviewCharacter = Cast<ACharacter>(PreviewActor))
+	if (bIsCharacter)
 	{
-		PreviewCharacter->bUseControllerRotationYaw = false;
+		ACharacter* PreviewCharacter = GetWorld()->SpawnActor<ACharacter>(ACharacter::StaticClass(), SpawnTransform, Params);
+		PreviewCharacter->SetFlags(RF_Transient);
+
+		USkeletalMeshComponent* SkelMeshComp = PreviewCharacter->GetMesh();
+		SkelMeshComp->SetRelativeLocation(FVector(0.f, 0.f, -PreviewCharacter->GetCapsuleComponent()->GetScaledCapsuleHalfHeight()));
+		SkelMeshComp->SetRelativeRotation(RoleDef->MeshToComponent.GetRotation());
+		SkelMeshComp->SetSkeletalMesh(Skeleton->GetPreviewMesh(true));
+		SkelMeshComp->SetAnimationMode(EAnimationMode::AnimationBlueprint);
+		SkelMeshComp->SetAnimInstanceClass(UAnimPreviewInstance::StaticClass());
+
+		UAnimPreviewInstance* AnimInstance = Cast<UAnimPreviewInstance>(SkelMeshComp->GetAnimInstance());
+		AnimInstance->SetAnimationAsset(AnimTrack.Animation, false, 0.f);
+		AnimInstance->PlayAnim(false, 0.f);
+
+		PreviewCharacter->CacheInitialMeshOffset(SkelMeshComp->GetRelativeLocation(), SkelMeshComp->GetRelativeRotation());
 
 		if (UCharacterMovementComponent* CharacterMovementComp = PreviewCharacter->GetCharacterMovement())
 		{
@@ -149,16 +178,38 @@ AActor* FContextualAnimViewModel::SpawnPreviewActor(const FContextualAnimTrack&
 			CharacterMovementComp->SetMovementMode(AnimTrack.bRequireFlyingMode ? EMovementMode::MOVE_Flying : EMovementMode::MOVE_Walking);
 		}
 
-		if (UCameraComponent* CameraComp = PreviewCharacter->FindComponentByClass<UCameraComponent>())
-		{
-			CameraComp->DestroyComponent();
-		}
+		UMotionWarpingComponent* MotionWarpingComp = NewObject<UMotionWarpingComponent>(PreviewCharacter);
+		MotionWarpingComp->RegisterComponentWithWorld(GetWorld());
+		MotionWarpingComp->InitializeComponent();
+
+		return PreviewCharacter;
 	}
+	else
+	{
+		AActor* PreviewActor = GetWorld()->SpawnActor<AActor>(AActor::StaticClass(), SpawnTransform, Params);
+		PreviewActor->SetFlags(RF_Transient);
+
+		UDebugSkelMeshComponent* SkelMeshComp = NewObject<UDebugSkelMeshComponent>(PreviewActor);
+		SkelMeshComp->RegisterComponentWithWorld(GetWorld());
+
+		UAnimPreviewInstance* AnimInstance = NewObject<UAnimPreviewInstance>(SkelMeshComp);
+		SkelMeshComp->PreviewInstance = AnimInstance;
+		AnimInstance->InitializeAnimation();
+
+		SkelMeshComp->SetSkeletalMesh(Skeleton->GetPreviewMesh(true));
+		SkelMeshComp->EnablePreview(true, AnimTrack.Animation);
 
-	UE_LOG(LogContextualAnim, Log, TEXT("Spawned preview Actor: %s at Loc: %s Rot: %s Role: %s"),
-		*GetNameSafe(PreviewActor), *SpawnTransform.GetLocation().ToString(), *SpawnTransform.Rotator().ToString(), *AnimTrack.Role.ToString());
+		AnimInstance->SetAnimationAsset(AnimTrack.Animation, false, 0.0f);
+
+		AnimInstance->PlayAnim(false, 0.0f);
+
+		if (!PreviewActor->GetRootComponent())
+		{
+			PreviewActor->SetRootComponent(SkelMeshComp);
+		}
 
-	return PreviewActor;
+		return PreviewActor;
+	}
 }
 
 void FContextualAnimViewModel::RefreshSequencerTracks()
@@ -181,7 +232,7 @@ void FContextualAnimViewModel::RefreshSequencerTracks()
 
 		MovieScene->RemoveMasterTrack(MasterTrack);
 	}
-			
+
 	if (SceneInstance.IsValid())
 	{
 		SceneInstance->Stop();
@@ -566,25 +617,196 @@ void FContextualAnimViewModel::AnimationModified(UAnimSequenceBase& Animation)
 	Animation.MarkPackageDirty();
 }
 
-void FContextualAnimViewModel::OnPreviewActorClassChanged()
+void FContextualAnimViewModel::UpdateSelection(const AActor* SelectedActor)
 {
-	const UContextualAnimRolesAsset* RolesAsset = GetSceneAsset()->RolesAsset;
+	const FContextualAnimSceneBinding* Binding = SceneInstance.IsValid() ? SceneInstance->FindBindingByActor(SelectedActor) : nullptr;
+	if (Binding)
+	{
+		UpdateSelection(Binding->GetRoleDef().Name);
+	}
+	else
+	{
+		ClearSelection();
+	}
+}
+
+void FContextualAnimViewModel::UpdateSelection(FName Role, int32 CriterionIdx, int32 CriterionDataIdx)
+{
+	Sequencer->EmptySelection();
+
+	SelectionInfo.Role = Role;
+	SelectionInfo.Criterion.Key = CriterionIdx;
+	SelectionInfo.Criterion.Value = CriterionDataIdx;
+}
+
+void FContextualAnimViewModel::ClearSelection()
+{
+	SelectionInfo.Reset();
+}
 
-	if(RolesAsset && SceneInstance.IsValid())
+FContextualAnimSceneBinding* FContextualAnimViewModel::GetSelectedBinding() const
+{
+	return SceneInstance.IsValid() ? const_cast<FContextualAnimSceneBinding*>(SceneInstance->FindBindingByRole(SelectionInfo.Role)) : nullptr;
+}
+
+AActor* FContextualAnimViewModel::GetSelectedActor() const
+{
+	const FContextualAnimSceneBinding* Binding = GetSelectedBinding();
+	return Binding ? Binding->GetActor() : nullptr;
+}
+
+FContextualAnimTrack* FContextualAnimViewModel::GetSelectedAnimTrack() const
+{
+	FContextualAnimSceneBinding* Binding = GetSelectedBinding();
+	return Binding ? const_cast<FContextualAnimTrack*>(&Binding->GetAnimTrack()) : nullptr;
+}
+
+UContextualAnimSelectionCriterion* FContextualAnimViewModel::GetSelectedSelectionCriterion() const
+{
+	if (SelectionInfo.Criterion.Key != INDEX_NONE && SelectionInfo.Criterion.Value != INDEX_NONE)
 	{
-		for (const auto& Binding : SceneInstance->GetBindings())
+		FContextualAnimTrack* AnimTrack = GetSelectedAnimTrack();
+		if (AnimTrack && AnimTrack->SelectionCriteria.IsValidIndex(SelectionInfo.Criterion.Key))
 		{
-			if (const FContextualAnimRoleDefinition* RoleDef = RolesAsset->FindRoleDefinitionByName(Binding.GetRoleDef().Name))
+			return AnimTrack->SelectionCriteria[SelectionInfo.Criterion.Key];
+		}
+	}
+
+	return nullptr;
+}
+
+FText FContextualAnimViewModel::GetSelectionDebugText() const
+{
+	AActor* SelectedActor = GetSelectedActor();
+	return FText::FromString(FString::Printf(TEXT("Selection Info:\n Role: %s \n Actor: %s \n Criterion: %d (%d)"),
+		*SelectionInfo.Role.ToString(), *GetNameSafe(SelectedActor), SelectionInfo.Criterion.Key, SelectionInfo.Criterion.Value));
+}
+
+bool FContextualAnimViewModel::ProcessInputDelta(FVector& InDrag, FRotator& InRot, FVector& InScale)
+{
+	if(IsSimulateModeActive())
+	{
+		if (AActor* SelectedActor = GetSelectedActor())
+		{
+			SelectedActor->SetActorLocationAndRotation(SelectedActor->GetActorLocation() + InDrag, SelectedActor->GetActorRotation() + InRot);
+			return true;
+		}
+	}
+	else
+	{
+		if (UContextualAnimSelectionCriterion_TriggerArea* Spatial = Cast<UContextualAnimSelectionCriterion_TriggerArea>(GetSelectedSelectionCriterion()))
+		{
+			FMatrix WidgetCoordSystem = FMatrix::Identity;
+			GetCustomDrawingCoordinateSystem(WidgetCoordSystem, nullptr);
+
+			InDrag = WidgetCoordSystem.InverseTransformVector(InDrag);
+
+			FVector& Point = Spatial->PolygonPoints[SelectionInfo.Criterion.Value >= 4 ? SelectionInfo.Criterion.Value - 4 : SelectionInfo.Criterion.Value];
+			Point.X += InDrag.X;
+			Point.Y += InDrag.Y;
+
+			if (InDrag.Z != 0.f)
 			{
-				const UClass* DesiredPreviewClass = RoleDef->PreviewActorClass;
-				const UClass* CurrentPreviewClass = Binding.GetActor()->GetClass();
+				if (SelectionInfo.Criterion.Value < 4)
+				{
+					for (int32 Idx = 0; Idx < Spatial->PolygonPoints.Num(); Idx++)
+					{
+						Spatial->PolygonPoints[Idx].Z += InDrag.Z;
+					}
 
-				if (DesiredPreviewClass && DesiredPreviewClass != CurrentPreviewClass)
+					Spatial->Height = FMath::Max(Spatial->Height - InDrag.Z, 0.f);
+				}
+				else
 				{
-					RefreshSequencerTracks();
-					break;
+					Spatial->Height = FMath::Max(Spatial->Height + InDrag.Z, 0.f);
 				}
 			}
+
+			return true;
+		}
+		else if (FContextualAnimSceneBinding* Binding = GetSelectedBinding())
+		{
+			FTransform& MeshToSpaceTransform = (const_cast<FContextualAnimTrack*>(&Binding->GetAnimTrack()))->MeshToScene;
+			MeshToSpaceTransform.SetLocation(MeshToSpaceTransform.GetLocation() + InDrag);
+			MeshToSpaceTransform.SetRotation(FQuat(InRot) * MeshToSpaceTransform.GetRotation());
+
+			UpdatePreviewActorTransform(*Binding, Sequencer->GetGlobalTime().AsSeconds());
+
+			return true;
+		}
+	}
+
+	return false;
+}
+
+bool FContextualAnimViewModel::GetCustomDrawingCoordinateSystem(FMatrix& InMatrix, void* InData)
+{
+	if (SelectionInfo.Criterion.Value != INDEX_NONE)
+	{
+		if (SceneInstance.IsValid())
+		{
+			FTransform PrimaryActorTransform = FTransform::Identity;
+			if (const FContextualAnimSceneBinding* Binding = SceneInstance->FindBindingByRole(GetSceneAsset()->GetPrimaryRole()))
+			{
+				PrimaryActorTransform = Binding->GetTransform();
+				InMatrix = PrimaryActorTransform.ToMatrixNoScale().RemoveTranslation();
+				return true;
+			}
 		}
 	}
+	else if (AActor* SelectedActor = GetSelectedActor())
+	{
+		InMatrix = SelectedActor->GetActorTransform().ToMatrixNoScale().RemoveTranslation();
+		return true;
+	}
+
+	return false;
+}
+
+bool FContextualAnimViewModel::ShouldPreviewSceneDrawWidget() const
+{
+	const UContextualAnimSelectionCriterion* SelectionCriterion = GetSelectedSelectionCriterion();
+	if (SelectionCriterion && SelectionCriterion->GetClass()->IsChildOf<UContextualAnimSelectionCriterion_TriggerArea>())
+	{
+		return true;
+	}
+	else if (GetSelectedActor())
+	{
+		return true;
+	}
+
+	return false;
+}
+
+FVector FContextualAnimViewModel::GetWidgetLocationFromSelection() const
+{
+	if (SelectionInfo.Criterion.Value != INDEX_NONE)
+	{
+		if (const UContextualAnimSelectionCriterion_TriggerArea* Spatial = Cast<UContextualAnimSelectionCriterion_TriggerArea>(GetSelectedSelectionCriterion()))
+		{
+			FVector Location = FVector::ZeroVector;
+			if (SelectionInfo.Criterion.Value < 4)
+			{
+				Location = Spatial->PolygonPoints[SelectionInfo.Criterion.Value];
+			}
+			else
+			{
+				Location = Spatial->PolygonPoints[SelectionInfo.Criterion.Value - 4] + FVector::UpVector * Spatial->Height;
+			}
+
+			FTransform PrimaryActorTransform = FTransform::Identity;
+			if (const FContextualAnimSceneBinding* Binding = SceneInstance->FindBindingByRole(SceneAsset->GetPrimaryRole()))
+			{
+				PrimaryActorTransform = Binding->GetTransform();
+			}
+
+			return PrimaryActorTransform.TransformPositionNoScale(Location);
+		}
+	}
+	else if (AActor* SelectedActor = GetSelectedActor())
+	{
+		return SelectedActor->GetActorLocation();
+	}
+
+	return FVector::ZeroVector;
 }
\ No newline at end of file
diff --git a/Engine/Plugins/Experimental/Animation/ContextualAnimation/Source/ContextualAnimationEditor/Public/ContextualAnimEdMode.h b/Engine/Plugins/Experimental/Animation/ContextualAnimation/Source/ContextualAnimationEditor/Public/ContextualAnimEdMode.h
--- a/Engine/Plugins/Experimental/Animation/ContextualAnimation/Source/ContextualAnimationEditor/Public/ContextualAnimEdMode.h
+++ b/Engine/Plugins/Experimental/Animation/ContextualAnimation/Source/ContextualAnimationEditor/Public/ContextualAnimEdMode.h
@@ -1,93 +1,59 @@
-// Copyright Epic Games, Inc. All Rights Reserved.
-
-#pragma once
-
-#include "CoreMinimal.h"
-#include "EdMode.h"
-
-class AActor;
-class FContextualAnimViewModel;
-
-struct FSelectionCriterionHitProxyData
-{
-	int32 SectionIdx = INDEX_NONE;
-	int32 AnimSetIdx = INDEX_NONE;
-	FName RoleName = NAME_None;
-	int32 CriterionIdx = INDEX_NONE;
-	int32 DataIdx = INDEX_NONE;
-
-	FSelectionCriterionHitProxyData(){}
-	FSelectionCriterionHitProxyData(int32 InSectionIdx, int32 InAnimSetIdx, const FName& InRoleName, int32 InCriterionIdx, int32 InDataIdx)
-		: SectionIdx(InSectionIdx)
-		, AnimSetIdx(InAnimSetIdx)
-		, RoleName(InRoleName)
-		, CriterionIdx(InCriterionIdx)
-		, DataIdx(InDataIdx)
-	{}
-
-	void Reset()
-	{
-		RoleName = NAME_None;
-		SectionIdx = INDEX_NONE;
-		AnimSetIdx = INDEX_NONE;
-		CriterionIdx = INDEX_NONE;
-		DataIdx = INDEX_NONE;
-	}
-
-	bool IsValid() const 
-	{ 
-		return	RoleName != NAME_None && 
-				SectionIdx != INDEX_NONE &&
-				AnimSetIdx != INDEX_NONE && 
-				CriterionIdx != INDEX_NONE; 
-	}
-};
-
-struct HSelectionCriterionHitProxy : public HHitProxy
-{
-	DECLARE_HIT_PROXY();
-
-	HSelectionCriterionHitProxy(const FSelectionCriterionHitProxyData& InData, EHitProxyPriority InPriority = HPP_Wireframe)
-		: HHitProxy(InPriority)
-		, Data(InData)
-	{}
-
-	virtual EMouseCursor::Type GetMouseCursor() override { return EMouseCursor::Crosshairs; }
-
-	FSelectionCriterionHitProxyData Data;
-};
-
-class FContextualAnimEdMode : public FEdMode
-{
-public:
-
-	const static FEditorModeID EdModeId;
-
-	FContextualAnimEdMode();
-	virtual ~FContextualAnimEdMode();
-
-	virtual void Tick(FEditorViewportClient* ViewportClient, float DeltaTime) override;
-	virtual void Render(const FSceneView* View, FViewport* Viewport, FPrimitiveDrawInterface* PDI);
-	virtual bool HandleClick(FEditorViewportClient* InViewportClient, HHitProxy* HitProxy, const FViewportClick& Click) override;
-	virtual bool StartTracking(FEditorViewportClient* InViewportClient, FViewport* InViewport) override;
-	virtual bool EndTracking(FEditorViewportClient* InViewportClient, FViewport* InViewport) override;
-	virtual bool InputDelta(FEditorViewportClient* InViewportClient, FViewport* InViewport, FVector& InDrag, FRotator& InRot, FVector& InScale) override;
-	virtual bool InputKey(FEditorViewportClient* ViewportClient, FViewport* Viewport, FKey Key, EInputEvent Event) override;
-	virtual bool AllowWidgetMove() override;
-	virtual bool ShouldDrawWidget() const override;
-	virtual FVector GetWidgetLocation() const override;
-	virtual bool GetCustomDrawingCoordinateSystem(FMatrix& InMatrix, void* InData);
-	virtual bool GetCustomInputCoordinateSystem(FMatrix& InMatrix, void* InData);
-
-	bool GetHitResultUnderCursor(FHitResult& OutHitResult, FEditorViewportClient* InViewportClient, const FViewportClick& Click) const;
-
-	void DrawIKTargetsForBinding(class FPrimitiveDrawInterface& PDI, const struct FContextualAnimSceneBinding& Binding) const;
-
-private:
-
-	FContextualAnimViewModel* ViewModel = nullptr;
-
-	TWeakObjectPtr<AActor> SelectedActor;
-
-	FSelectionCriterionHitProxyData SelectedSelectionCriterionData;
-};
+// Copyright Epic Games, Inc. All Rights Reserved.
+
+#pragma once
+
+#include "CoreMinimal.h"
+#include "EdMode.h"
+
+class AActor;
+class FContextualAnimViewModel;
+
+struct HSelectionCriterionHitProxy : public HHitProxy
+{
+	DECLARE_HIT_PROXY();
+
+	HSelectionCriterionHitProxy(FName InRole, int32 InCriterionIdx, int32 InDataIdx, EHitProxyPriority InPriority = HPP_Wireframe)
+		: HHitProxy(InPriority)
+	{
+		Role = InRole;
+		IndexPair.Key = InCriterionIdx;
+		IndexPair.Value = InDataIdx;
+	}
+
+	virtual EMouseCursor::Type GetMouseCursor() override { return EMouseCursor::Crosshairs; }
+
+	FName Role = NAME_None;
+	TPair<int32, int32> IndexPair;
+};
+
+class FContextualAnimEdMode : public FEdMode
+{
+public:
+
+	const static FEditorModeID EdModeId;
+
+	FContextualAnimEdMode();
+	virtual ~FContextualAnimEdMode();
+
+	virtual void Tick(FEditorViewportClient* ViewportClient, float DeltaTime) override;
+	virtual void Render(const FSceneView* View, FViewport* Viewport, FPrimitiveDrawInterface* PDI);
+	virtual void DrawHUD(FEditorViewportClient* ViewportClient, FViewport* Viewport, const FSceneView* View, FCanvas* Canvas) override;
+	virtual bool HandleClick(FEditorViewportClient* InViewportClient, HHitProxy* HitProxy, const FViewportClick& Click) override;
+	virtual bool StartTracking(FEditorViewportClient* InViewportClient, FViewport* InViewport) override;
+	virtual bool EndTracking(FEditorViewportClient* InViewportClient, FViewport* InViewport) override;
+	virtual bool InputDelta(FEditorViewportClient* InViewportClient, FViewport* InViewport, FVector& InDrag, FRotator& InRot, FVector& InScale) override;
+	virtual bool InputKey(FEditorViewportClient* ViewportClient, FViewport* Viewport, FKey Key, EInputEvent Event) override;
+	virtual bool AllowWidgetMove() override;
+	virtual bool ShouldDrawWidget() const override;
+	virtual FVector GetWidgetLocation() const override;
+	virtual bool GetCustomDrawingCoordinateSystem(FMatrix& InMatrix, void* InData);
+	virtual bool GetCustomInputCoordinateSystem(FMatrix& InMatrix, void* InData);
+
+	bool GetHitResultUnderCursor(FHitResult& OutHitResult, FEditorViewportClient* InViewportClient, const FViewportClick& Click) const;
+
+	void DrawIKTargetsForBinding(class FPrimitiveDrawInterface& PDI, const struct FContextualAnimSceneBinding& Binding) const;
+
+private:
+
+	FContextualAnimViewModel* ViewModel = nullptr;
+};
diff --git a/Engine/Plugins/Experimental/Animation/ContextualAnimation/Source/ContextualAnimationEditor/Public/ContextualAnimViewModel.h b/Engine/Plugins/Experimental/Animation/ContextualAnimation/Source/ContextualAnimationEditor/Public/ContextualAnimViewModel.h
--- a/Engine/Plugins/Experimental/Animation/ContextualAnimation/Source/ContextualAnimationEditor/Public/ContextualAnimViewModel.h
+++ b/Engine/Plugins/Experimental/Animation/ContextualAnimation/Source/ContextualAnimationEditor/Public/ContextualAnimViewModel.h
@@ -62,16 +62,54 @@ public:
 
 	void SetActiveAnimSetForSection(int32 SectionIdx, int32 AnimSetIdx);
 
-	void OnPreviewActorClassChanged();
-
 	void ToggleSimulateMode();
 	bool IsSimulateModeActive() { return bIsSimulateModeActive; }
 	void StartSimulation();
 
 	void UpdatePreviewActorTransform(const FContextualAnimSceneBinding& Binding, float Time);
 
+	void UpdateSelection(const AActor* SelectedActor);
+	void UpdateSelection(FName Role, int32 CriterionIdx = INDEX_NONE, int32 CriterionDataIdx = INDEX_NONE);
+	void ClearSelection();
+	FContextualAnimSceneBinding* GetSelectedBinding() const;
+	FContextualAnimTrack* GetSelectedAnimTrack() const;
+	AActor* GetSelectedActor() const;
+	class UContextualAnimSelectionCriterion* GetSelectedSelectionCriterion() const;
+	FText GetSelectionDebugText() const;
+
+	bool ProcessInputDelta(FVector& InDrag, FRotator& InRot, FVector& InScale);
+	bool ShouldPreviewSceneDrawWidget() const;
+	bool GetCustomDrawingCoordinateSystem(FMatrix& InMatrix, void* InData);
+	FVector GetWidgetLocationFromSelection() const;
+
 private:
 
+	struct FContextualAnimEdSelectionInfo
+	{
+		/** Selected Role */
+		FName Role = NAME_None;
+
+		/** Key = Criterion Idx in the AnimTrack, Value = Data idx (e.g vertex idx) */
+		TPair<int32, int32> Criterion;
+
+		/** FGuid of the selected AnimNotify */
+		FGuid NotifyGuid;
+
+		FContextualAnimEdSelectionInfo()
+		{
+			Reset();
+		}
+
+		void Reset()
+		{
+			Role = NAME_None;
+			Criterion.Key = INDEX_NONE;
+			Criterion.Value = INDEX_NONE;
+			NotifyGuid = FGuid();
+		}
+	};
+	FContextualAnimEdSelectionInfo SelectionInfo;
+
 	/** Scene asset being viewed and edited by this view model. */
 	TObjectPtr<UContextualAnimSceneAsset> SceneAsset;
 
