Contextual Anim:
- Added type property to SelectionCriterion (defaulted to Spatial)
- Added query functions to get specific alignment points from a given section.
- Added AnimTrackIdx to AnimTrack
- Moved Radius property from each AnimSet to the SceneAsset (manually set for now)

diff --git a/Engine/Plugins/Experimental/Animation/ContextualAnimation/Source/ContextualAnimation/Private/ContextualAnimSceneActorComponent.cpp b/Engine/Plugins/Experimental/Animation/ContextualAnimation/Source/ContextualAnimation/Private/ContextualAnimSceneActorComponent.cpp
--- a/Engine/Plugins/Experimental/Animation/ContextualAnimation/Source/ContextualAnimation/Private/ContextualAnimSceneActorComponent.cpp
+++ b/Engine/Plugins/Experimental/Animation/ContextualAnimation/Source/ContextualAnimation/Private/ContextualAnimSceneActorComponent.cpp
@@ -28,9 +28,8 @@
 
 FBoxSphereBounds UContextualAnimSceneActorComponent::CalcBounds(const FTransform& LocalToWorld) const
 {
-	// Use radius of the first set in the first section for now. 
 	// The option of having an SceneAsset and draw options on this component may go away in the future anyway, replaced by smart objects.
-	const float Radius = SceneAsset && SceneAsset->HasValidData() ? SceneAsset->GetSection(0)->GetAnimSet(0)->Radius : 0.f;
+	const float Radius = SceneAsset && SceneAsset->HasValidData() ? SceneAsset->GetRadius() : 0.f;
 	return FBoxSphereBounds(FSphere(GetComponentTransform().GetLocation(), Radius));
 }
 
diff --git a/Engine/Plugins/Experimental/Animation/ContextualAnimation/Source/ContextualAnimation/Private/ContextualAnimSceneAsset.cpp b/Engine/Plugins/Experimental/Animation/ContextualAnimation/Source/ContextualAnimation/Private/ContextualAnimSceneAsset.cpp
--- a/Engine/Plugins/Experimental/Animation/ContextualAnimation/Source/ContextualAnimation/Private/ContextualAnimSceneAsset.cpp
+++ b/Engine/Plugins/Experimental/Animation/ContextualAnimation/Source/ContextualAnimation/Private/ContextualAnimSceneAsset.cpp
@@ -455,24 +455,25 @@ void UContextualAnimSceneAsset::PreSave(FObjectPreSaveContext ObjectSaveContext)
 
 void UContextualAnimSceneAsset::PrecomputeData()
 {
-	// Set SetIdx and SectionIdx on each AnimTrack and update the Radius for each set
+	if(!HasValidData())
+	{
+		return;
+	}
+
+	// Set indices on each AnimTrack
 	for (int32 SectionIdx = 0; SectionIdx < Sections.Num(); SectionIdx++)
 	{
 		FContextualAnimSceneSection& Section = Sections[SectionIdx];
 		for (int32 SetIdx = 0; SetIdx < Section.AnimSets.Num(); SetIdx++)
 		{
 			FContextualAnimSet& AnimSet = Section.AnimSets[SetIdx];
-
-			FBox Box(ForceInitToZero);
-			for (FContextualAnimTrack& AnimTrack : AnimSet.Tracks)
+			for (int32 TrackIdx = 0; TrackIdx < AnimSet.Tracks.Num(); TrackIdx++)
 			{
+				FContextualAnimTrack& AnimTrack = AnimSet.Tracks[TrackIdx];
 				AnimTrack.SectionIdx = SectionIdx;
 				AnimTrack.AnimSetIdx = SetIdx;
-
-				Box += AnimTrack.GetRootTransformAtTime(0.f).GetLocation();
+				AnimTrack.AnimTrackIdx = TrackIdx;
 			}
-
-			AnimSet.Radius = Box.GetExtent().Size();
 		}
 
 		// Generate alignment tracks
@@ -636,6 +637,99 @@ const TArray<FContextualAnimSetPivotDefinition>& UContextualAnimSceneAsset::GetA
 	return Sections.IsValidIndex(SectionIdx) ? Sections[SectionIdx].GetAnimSetPivotDefinitions() : EmptyDefs;
 }
 
+static FContextualAnimPoint GetContextualAnimPoint(const FContextualAnimTrack& AnimTrack, const FTransform& ToWorldTransform, int32 SampleRate, EContextualAnimPointType Type)
+{
+	check(SampleRate > 0);
+
+	if (AnimTrack.Animation)
+	{
+		const float Interval = 1.f / SampleRate;
+
+		float T1 = 0.f;
+		if(Type == EContextualAnimPointType::SyncFrame)
+		{
+			T1 = AnimTrack.GetSyncTimeForWarpSection(0);
+		}
+		else if(Type == EContextualAnimPointType::LastFrame)
+		{
+			T1 = FMath::Max(AnimTrack.Animation->GetPlayLength() - Interval, 0.f);
+		}
+
+		float T2 = T1 + Interval;
+
+		const FTransform RootTransform1 = AnimTrack.GetAlignmentTransformAtTime(T1);
+		const FTransform RootTransform2 = AnimTrack.GetAlignmentTransformAtTime(T2);
+		const float Delta = (RootTransform2.GetTranslation() - RootTransform1.GetTranslation()).Size();
+		const float Speed = Delta / Interval;
+
+		return FContextualAnimPoint(RootTransform1 * ToWorldTransform, Speed, AnimTrack.SectionIdx, AnimTrack.AnimSetIdx, AnimTrack.AnimTrackIdx);
+	}
+	else
+	{
+		return FContextualAnimPoint(ToWorldTransform, 0.f, AnimTrack.SectionIdx, AnimTrack.AnimSetIdx, AnimTrack.AnimTrackIdx);
+	}
+}
+
+void UContextualAnimSceneAsset::GetAlignmentPointsForRoleInSection(EContextualAnimPointType Type, int32 SectionIdx, const FName& Role, const FContextualAnimSceneBindingContext& Primary, TArray<FContextualAnimPoint>& OutResult) const
+{
+	if (Sections.IsValidIndex(SectionIdx))
+	{
+		OutResult.Reset(Sections[SectionIdx].AnimSets.Num());
+		for (const FContextualAnimSet& Set : Sections[SectionIdx].AnimSets)
+		{
+			for (const FContextualAnimTrack& AnimTrack : Set.Tracks)
+			{
+				if (AnimTrack.Role == Role)
+				{
+					OutResult.Add(GetContextualAnimPoint(AnimTrack, Primary.GetTransform(), GetSampleRate(), Type));
+					break;
+				}
+			}
+		}
+	}
+}
+
+void UContextualAnimSceneAsset::GetAlignmentPointsForRoleInSectionConsideringSelectionCriteria(EContextualAnimPointType Type, int32 SectionIdx, const FName& Role, const FContextualAnimSceneBindingContext& Querier, const FContextualAnimSceneBindingContext& Primary, EContextualAnimCriterionToConsider CriterionToConsider, TArray<FContextualAnimPoint>& OutResult) const
+{
+	if (Sections.IsValidIndex(SectionIdx))
+	{
+		OutResult.Reset(Sections[SectionIdx].AnimSets.Num());
+		for (const FContextualAnimSet& Set : Sections[SectionIdx].AnimSets)
+		{
+			for (const FContextualAnimTrack& AnimTrack : Set.Tracks)
+			{
+				if (AnimTrack.Role == Role)
+				{
+					bool bSuccess = true;
+					for (const UContextualAnimSelectionCriterion* Criterion : AnimTrack.SelectionCriteria)
+					{
+						if (Criterion)
+						{
+							if ((CriterionToConsider == EContextualAnimCriterionToConsider::All) ||
+								(CriterionToConsider == EContextualAnimCriterionToConsider::Spatial && Criterion->Type == EContextualAnimCriterionType::Spatial) ||
+								(CriterionToConsider == EContextualAnimCriterionToConsider::Other && Criterion->Type == EContextualAnimCriterionType::Other))
+							{
+								if (!Criterion->DoesQuerierPassCondition(Primary, Querier))
+								{
+									bSuccess = false;
+									break;
+								}
+							}
+						}
+					}
+
+					if (bSuccess)
+					{
+						OutResult.Add(GetContextualAnimPoint(AnimTrack, Primary.GetTransform(), GetSampleRate(), Type));
+					}
+
+					break;
+				}
+			}
+		}
+	}
+}
+
 // Blueprint Interface
 //------------------------------------------------------------------------------------------
 
diff --git a/Engine/Plugins/Experimental/Animation/ContextualAnimation/Source/ContextualAnimation/Public/ContextualAnimSceneAsset.h b/Engine/Plugins/Experimental/Animation/ContextualAnimation/Source/ContextualAnimation/Public/ContextualAnimSceneAsset.h
--- a/Engine/Plugins/Experimental/Animation/ContextualAnimation/Source/ContextualAnimation/Public/ContextualAnimSceneAsset.h
+++ b/Engine/Plugins/Experimental/Animation/ContextualAnimation/Source/ContextualAnimation/Public/ContextualAnimSceneAsset.h
@@ -42,9 +42,6 @@ struct CONTEXTUALANIMATION_API FContextualAnimSet
 
 	UPROPERTY(EditAnywhere, Category = "Defaults")
 	TArray<FTransform> ScenePivots;
-
-	UPROPERTY(VisibleAnywhere, Category = "Settings")
-	float Radius = 0.f;
 };
 
 /** Named container with one or more ContextualAnimSet */
@@ -97,6 +94,48 @@ protected:
  	friend class FContextualAnimViewModel;
 };
 
+USTRUCT(BlueprintType)
+struct FContextualAnimPoint
+{
+	GENERATED_BODY()
+
+	UPROPERTY()
+	FTransform Transform;
+
+	UPROPERTY()
+	float Speed = 0.f;
+
+	UPROPERTY()
+	int32 SectionIdx = INDEX_NONE;
+
+	UPROPERTY()
+	int32 AnimSetIdx = INDEX_NONE;
+
+	UPROPERTY()
+	int32 AnimTrackIdx = INDEX_NONE;
+
+	FContextualAnimPoint(){}
+	FContextualAnimPoint(const FTransform& InTransform, float InSpeed, int32 InSectionIdx, int32 InAnimSetIdx, int32 InAnimTrackIdx)
+		: Transform(InTransform), Speed(InSpeed), SectionIdx(InSectionIdx), AnimSetIdx(InAnimSetIdx), AnimTrackIdx(InAnimTrackIdx)
+	{}
+};
+
+UENUM(BlueprintType)
+enum class EContextualAnimPointType : uint8
+{
+	FirstFrame,
+	SyncFrame,
+	LastFrame
+};
+
+UENUM(BlueprintType)
+enum class EContextualAnimCriterionToConsider : uint8
+{
+	All,
+	Spatial,
+	Other
+};
+
 UCLASS(Blueprintable)
 class CONTEXTUALANIMATION_API UContextualAnimSceneAsset : public UDataAsset
 {
@@ -119,7 +158,8 @@ public:
 	FORCEINLINE bool GetDisableCollisionBetweenActors() const { return bDisableCollisionBetweenActors; }
 	FORCEINLINE const TSubclassOf<UContextualAnimSceneInstance>& GetSceneInstanceClass() const { return SceneInstanceClass; }
 	FORCEINLINE int32 GetSampleRate() const { return SampleRate; }
-	
+	FORCEINLINE float GetRadius() const { return Radius; }
+
 	bool HasValidData() const { return RolesAsset != nullptr && Sections.Num() > 0 && Sections[0].AnimSets.Num() > 0; }
 
 	const UContextualAnimRolesAsset* GetRolesAsset() const { return RolesAsset; }
@@ -155,6 +195,12 @@ public:
 
 	const FContextualAnimIKTargetDefContainer& GetIKTargetDefsForRoleInSection(int32 SectionIdx, const FName& Role) const;
 
+	void GetAlignmentPointsForRoleInSection(EContextualAnimPointType Type, int32 SectionIdx, const FName& Role, const FContextualAnimSceneBindingContext& Primary, TArray<FContextualAnimPoint>& OutResult) const;
+
+	void GetAlignmentPointsForRoleInSectionConsideringSelectionCriteria(EContextualAnimPointType Type, int32 SectionIdx, const FName& Role, const FContextualAnimSceneBindingContext& Querier, const FContextualAnimSceneBindingContext& Primary, EContextualAnimCriterionToConsider CriterionToConsider, TArray<FContextualAnimPoint>& OutResult) const;
+
+public:
+
 	// Blueprint Interface
 	//------------------------------------------------------------------------------------------
 
@@ -185,6 +231,9 @@ protected:
 	UPROPERTY(EditAnywhere, Category = "Defaults")
 	TArray<FContextualAnimSceneSection> Sections;
 
+	UPROPERTY(EditAnywhere, Category = "Settings")
+	float Radius = 0.f;
+
 	UPROPERTY(EditAnywhere, Category = "Settings")
 	TSubclassOf<UContextualAnimSceneInstance> SceneInstanceClass;
 
diff --git a/Engine/Plugins/Experimental/Animation/ContextualAnimation/Source/ContextualAnimation/Public/ContextualAnimSelectionCriterion.h b/Engine/Plugins/Experimental/Animation/ContextualAnimation/Source/ContextualAnimation/Public/ContextualAnimSelectionCriterion.h
--- a/Engine/Plugins/Experimental/Animation/ContextualAnimation/Source/ContextualAnimation/Public/ContextualAnimSelectionCriterion.h
+++ b/Engine/Plugins/Experimental/Animation/ContextualAnimation/Source/ContextualAnimation/Public/ContextualAnimSelectionCriterion.h
@@ -7,6 +7,13 @@
 #include "ContextualAnimTypes.h"
 #include "ContextualAnimSelectionCriterion.generated.h"
 
+UENUM(BlueprintType)
+enum class EContextualAnimCriterionType : uint8
+{
+	Spatial,
+	Other
+};
+
 // UContextualAnimSelectionCriterion
 //===========================================================================
 
@@ -17,6 +24,9 @@ class CONTEXTUALANIMATION_API UContextualAnimSelectionCriterion : public UObject
 
 public:
 
+	UPROPERTY(EditAnywhere, BlueprintReadOnly, Category = "Defaults")
+	EContextualAnimCriterionType Type = EContextualAnimCriterionType::Spatial;
+
 	UContextualAnimSelectionCriterion(const FObjectInitializer& ObjectInitializer);
 
 	class UContextualAnimSceneAsset* GetSceneAssetOwner() const;
