[ContextualAnim] manager from game worlds should not tick in Editor so scenes could be paused or stepped frame by frame
#rnx

diff --git a/Engine/Plugins/Experimental/Animation/ContextualAnimation/Source/ContextualAnimation/Private/ContextualAnimManager.cpp b/Engine/Plugins/Experimental/Animation/ContextualAnimation/Source/ContextualAnimation/Private/ContextualAnimManager.cpp
--- a/Engine/Plugins/Experimental/Animation/ContextualAnimation/Source/ContextualAnimation/Private/ContextualAnimManager.cpp
+++ b/Engine/Plugins/Experimental/Animation/ContextualAnimation/Source/ContextualAnimation/Private/ContextualAnimManager.cpp
@@ -41,6 +41,12 @@ TStatId UContextualAnimManager::GetStatId() const
 	RETURN_QUICK_DECLARE_CYCLE_STAT(UContextualAnimManager, STATGROUP_Tickables);
 }
 
+bool UContextualAnimManager::IsTickableInEditor() const
+{
+	const UWorld* World = GetWorld();
+	return World != nullptr && !World->IsGameWorld();
+}
+
 void UContextualAnimManager::RegisterSceneActorComponent(UContextualAnimSceneActorComponent* SceneActorComp)
 {
 	if(SceneActorComp)
diff --git a/Engine/Plugins/Experimental/Animation/ContextualAnimation/Source/ContextualAnimation/Public/ContextualAnimManager.h b/Engine/Plugins/Experimental/Animation/ContextualAnimation/Source/ContextualAnimation/Public/ContextualAnimManager.h
--- a/Engine/Plugins/Experimental/Animation/ContextualAnimation/Source/ContextualAnimation/Public/ContextualAnimManager.h
+++ b/Engine/Plugins/Experimental/Animation/ContextualAnimation/Source/ContextualAnimation/Public/ContextualAnimManager.h
@@ -32,7 +32,7 @@
 	virtual void Tick(float DeltaTime) override;
 	virtual ETickableTickType GetTickableTickType() const override;
 	virtual TStatId GetStatId() const override;
-	virtual bool IsTickableInEditor() const { return true; }
+	virtual bool IsTickableInEditor() const override;
 	// FTickableGameObject end
 
 	static UContextualAnimManager* Get(const UWorld* World);
