Contextual Anim Editor: Re-added view to edit notifies but now in a second level

diff --git a/Engine/Plugins/Experimental/Animation/ContextualAnimation/Source/ContextualAnimationEditor/Private/ContextualAnimAssetEditorToolkit.cpp b/Engine/Plugins/Experimental/Animation/ContextualAnimation/Source/ContextualAnimationEditor/Private/ContextualAnimAssetEditorToolkit.cpp
--- a/Engine/Plugins/Experimental/Animation/ContextualAnimation/Source/ContextualAnimationEditor/Private/ContextualAnimAssetEditorToolkit.cpp
+++ b/Engine/Plugins/Experimental/Animation/ContextualAnimation/Source/ContextualAnimationEditor/Private/ContextualAnimAssetEditorToolkit.cpp
@@ -57,7 +57,7 @@ UContextualAnimSceneAsset* FContextualAnimAssetEditorToolkit::GetSceneAsset() co
 
 void FContextualAnimAssetEditorToolkit::ResetPreviewScene()
 {
-	ViewModel->RefreshSequencerTracks();
+	ViewModel->SetDefaultMode();
 }
 
 void FContextualAnimAssetEditorToolkit::InitAssetEditor(const EToolkitMode::Type Mode, const TSharedPtr< IToolkitHost >& InitToolkitHost, UContextualAnimSceneAsset* SceneAsset)
diff --git a/Engine/Plugins/Experimental/Animation/ContextualAnimation/Source/ContextualAnimationEditor/Private/ContextualAnimMovieSceneNotifyTrackEditor.cpp b/Engine/Plugins/Experimental/Animation/ContextualAnimation/Source/ContextualAnimationEditor/Private/ContextualAnimMovieSceneNotifyTrackEditor.cpp
--- a/Engine/Plugins/Experimental/Animation/ContextualAnimation/Source/ContextualAnimationEditor/Private/ContextualAnimMovieSceneNotifyTrackEditor.cpp
+++ b/Engine/Plugins/Experimental/Animation/ContextualAnimation/Source/ContextualAnimationEditor/Private/ContextualAnimMovieSceneNotifyTrackEditor.cpp
@@ -22,6 +22,8 @@
 #include "Modules/ModuleManager.h"
 #include "Widgets/Input/SButton.h"
 #include "ContextualAnimEditorTypes.h"
+#include "ISequencerModule.h"
+#include "Toolkits/AssetEditorToolkit.h"
 
 #define LOCTEXT_NAMESPACE "FContextualAnimMovieSceneNotifyTrackEditor"
 
@@ -60,6 +62,46 @@
 	return MakeShareable(new FContextualAnimNotifySection(SectionObject));
 }
 
+void FContextualAnimMovieSceneNotifyTrackEditor::OnInitialize()
+{
+	ISequencerModule& SequencerModule = FModuleManager::Get().LoadModuleChecked<ISequencerModule>("Sequencer");
+
+	TSharedPtr<FUICommandList> CommandList = MakeShareable(new FUICommandList);
+
+	// Extend the sequencer toolbar to add Back button
+	TSharedPtr<FExtensibilityManager> ExtensionManager = SequencerModule.GetToolBarExtensibilityManager();
+	TSharedPtr<FExtender> ToolbarExtender = MakeShareable(new FExtender);
+	TSharedPtr<const FExtensionBase> ToolbarExtension = ToolbarExtender->AddToolBarExtension(
+		"CurveEditor",
+		EExtensionHook::After,
+		CommandList,
+		FToolBarExtensionDelegate::CreateSP(this, &FContextualAnimMovieSceneNotifyTrackEditor::CustomizeToolBar)
+	);
+	ExtensionManager->AddExtender(ToolbarExtender);
+}
+
+void FContextualAnimMovieSceneNotifyTrackEditor::CustomizeToolBar(FToolBarBuilder& ToolBarBuilder)
+{
+	ToolBarBuilder.BeginSection("Custom");
+	{
+		FUIAction Action = FUIAction(FExecuteAction::CreateLambda([this]() {
+			GetMovieSceneSequence().GetViewModel().SetDefaultMode();
+		}));
+
+		Action.IsActionVisibleDelegate = FIsActionButtonVisible::CreateLambda([this]() {
+			return GetMovieSceneSequence().GetViewModel().GetTimelineMode() == FContextualAnimViewModel::ETimelineMode::Notifies;
+		});
+
+		ToolBarBuilder.AddToolBarButton(
+			Action,
+			NAME_None,
+			TAttribute<FText>(),
+			TAttribute<FText>(),
+			FSlateIcon(FAppStyle::GetAppStyleSetName(), "Icons.ArrowLeft"));
+	}
+	ToolBarBuilder.EndSection();
+}
+
 void FContextualAnimMovieSceneNotifyTrackEditor::FillNewNotifyStateMenu(FMenuBuilder& MenuBuilder, bool bIsReplaceWithMenu, UContextualAnimMovieSceneNotifyTrack* Track, int32 RowIndex)
 {
 	// MenuBuilder always has a search widget added to it by default, hence if larger then 1 then something else has been added to it
@@ -90,6 +132,17 @@
 void FContextualAnimMovieSceneNotifyTrackEditor::BuildAddTrackMenu(FMenuBuilder& MenuBuilder)
 {
 	// Menu that appears when clicking on the Add Track button next to the Search Tracks bar
+
+	if (GetMovieSceneSequence().GetViewModel().GetTimelineMode() == FContextualAnimViewModel::ETimelineMode::Notifies)
+	{
+		MenuBuilder.AddMenuEntry(
+			LOCTEXT("NotifyTrack", "Notify Track"),
+			LOCTEXT("NotifyTrackTooltip", "Adds a notify track"),
+			FSlateIcon(),
+			FUIAction(FExecuteAction::CreateRaw(this, &FContextualAnimMovieSceneNotifyTrackEditor::AddNewNotifyTrack)
+			)
+		);
+	}
 }
 
 TSharedPtr<SWidget> FContextualAnimMovieSceneNotifyTrackEditor::BuildOutlinerEditWidget(const FGuid& ObjectBinding, UMovieSceneTrack* Track, const FBuildEditWidgetParams& Params)
@@ -139,7 +192,7 @@
 		.AutoWidth()
 		.VAlign(VAlign_Center)
 		[
-			FSequencerUtilities::MakeAddButton(LOCTEXT("AddSection", "Section"), FOnGetContent::CreateLambda(SubMenuCallback), Params.NodeIsHovered, GetSequencer())
+			FSequencerUtilities::MakeAddButton(LOCTEXT("AddNotify", "Notify"), FOnGetContent::CreateLambda(SubMenuCallback), Params.NodeIsHovered, GetSequencer())
 		];
 }
 
@@ -261,21 +314,9 @@
 void FContextualAnimMovieSceneNotifyTrackEditor::BuildObjectBindingTrackMenu(FMenuBuilder& MenuBuilder, const TArray<FGuid>& ObjectBindings, const UClass* ObjectClass)
 {	
 	// Builds menu that appears when clicking on the +Track button on an Object Track
-
-	UMovieSceneSequence* MovieSequence = GetSequencer()->GetFocusedMovieSceneSequence();
-	if(MovieSequence && MovieSequence->GetMovieScene()->FindPossessable(ObjectBindings[0]))
-	{
-		MenuBuilder.AddMenuEntry(
-			LOCTEXT("NotifyTrack", "Notify Track"),
-			LOCTEXT("NotifyTrackTooltip", "Adds a notify track"),
-			FSlateIcon(),
-			FUIAction(FExecuteAction::CreateRaw(this, &FContextualAnimMovieSceneNotifyTrackEditor::AddNewNotifyTrack, ObjectBindings)
-			)
-		);
-	}
 }
 
-void FContextualAnimMovieSceneNotifyTrackEditor::AddNewNotifyTrack(TArray<FGuid> ObjectBindings)
+void FContextualAnimMovieSceneNotifyTrackEditor::AddNewNotifyTrack()
 {
 	// Copied from AnimTimelineTrack_Notifies.cpp/FAnimTimelineTrack_Notifies::GetNewTrackName()
 	auto GetNewTrackName = [](UAnimSequenceBase* InAnimSequenceBase) -> FName
@@ -299,30 +340,26 @@
 		return NameToTest;
 	};
 
-	// @TODO: Commented out for now until we add the new behavior where the user needs to double-click on the animation to edit the notifies
-	/*for (const FGuid& ObjectBinding : ObjectBindings)
-	{
- 		UAnimSequenceBase* Animation = GetMovieSceneSequence().GetViewModel().FindAnimationByGuid(ObjectBinding);
- 		check(Animation);
+	UAnimSequenceBase* Animation = GetMovieSceneSequence().GetViewModel().GetEditingAnimation();
+	check(Animation);
 
-		// Copied from AnimTimelineTrack_Notifies.cpp/FAnimTimelineTrack_Notifies::AddTrack()
+	// Copied from AnimTimelineTrack_Notifies.cpp/FAnimTimelineTrack_Notifies::AddTrack()
 
-		FAnimNotifyTrack NewNotifyTrack;
-		NewNotifyTrack.TrackName = GetNewTrackName(Animation);
-		NewNotifyTrack.TrackColor = FLinearColor::White;
+	FAnimNotifyTrack NewNotifyTrack;
+	NewNotifyTrack.TrackName = GetNewTrackName(Animation);
+	NewNotifyTrack.TrackColor = FLinearColor::White;
 
-		Animation->AnimNotifyTracks.Add(NewNotifyTrack);
+	Animation->AnimNotifyTracks.Add(NewNotifyTrack);
 
-		GetMovieSceneSequence().GetViewModel().AnimationModified(*Animation);
+	GetMovieSceneSequence().GetViewModel().AnimationModified(*Animation);
 
-		// Create and Initialize MovieSceneTrack
-		UContextualAnimMovieSceneNotifyTrack* MovieSceneTrack = GetMovieSceneSequence().GetMovieScene()->AddTrack<UContextualAnimMovieSceneNotifyTrack>(ObjectBinding);
-		check(MovieSceneTrack);
+	// Create and Initialize MovieSceneTrack
+	UContextualAnimMovieSceneNotifyTrack* MovieSceneTrack = GetMovieSceneSequence().GetMovieScene()->AddMasterTrack<UContextualAnimMovieSceneNotifyTrack>();
+	check(MovieSceneTrack);
 
-		MovieSceneTrack->Initialize(*Animation, NewNotifyTrack);
-	}
+	MovieSceneTrack->Initialize(*Animation, NewNotifyTrack);
 
-	GetSequencer()->NotifyMovieSceneDataChanged(EMovieSceneDataChangeType::MovieSceneStructureItemAdded);*/
+	GetSequencer()->NotifyMovieSceneDataChanged(EMovieSceneDataChangeType::MovieSceneStructureItemAdded);
 }
 
 bool FContextualAnimMovieSceneNotifyTrackEditor::SupportsType(TSubclassOf<UMovieSceneTrack> Type) const
diff --git a/Engine/Plugins/Experimental/Animation/ContextualAnimation/Source/ContextualAnimationEditor/Private/ContextualAnimMovieSceneTrackEditor.cpp b/Engine/Plugins/Experimental/Animation/ContextualAnimation/Source/ContextualAnimationEditor/Private/ContextualAnimMovieSceneTrackEditor.cpp
--- a/Engine/Plugins/Experimental/Animation/ContextualAnimation/Source/ContextualAnimationEditor/Private/ContextualAnimMovieSceneTrackEditor.cpp
+++ b/Engine/Plugins/Experimental/Animation/ContextualAnimation/Source/ContextualAnimationEditor/Private/ContextualAnimMovieSceneTrackEditor.cpp
@@ -77,7 +77,11 @@
 	{
 		if(Section->IsActive())
 		{
-			//@TODO: Jump to Edit Notifies view
+			FContextualAnimTrack& AnimTrack = Section->GetAnimTrack();
+			if(AnimTrack.Animation)
+			{
+				Section->GetOwnerTrack().GetViewModel().SetNotifiesMode(Section->GetAnimTrack());
+			}
 		}
 		else
 		{
diff --git a/Engine/Plugins/Experimental/Animation/ContextualAnimation/Source/ContextualAnimationEditor/Private/ContextualAnimViewModel.cpp b/Engine/Plugins/Experimental/Animation/ContextualAnimation/Source/ContextualAnimationEditor/Private/ContextualAnimViewModel.cpp
--- a/Engine/Plugins/Experimental/Animation/ContextualAnimation/Source/ContextualAnimationEditor/Private/ContextualAnimViewModel.cpp
+++ b/Engine/Plugins/Experimental/Animation/ContextualAnimation/Source/ContextualAnimationEditor/Private/ContextualAnimViewModel.cpp
@@ -74,7 +74,7 @@ void FContextualAnimViewModel::Initialize(UContextualAnimSceneAsset* InSceneAsse
 
 	CreateSequencer();
 
-	RefreshSequencerTracks();
+	SetDefaultMode();
 }
 
 void FContextualAnimViewModel::CreateSequencer()
@@ -114,7 +114,7 @@ void FContextualAnimViewModel::SetActiveSection(int32 SectionIdx)
 
 	ActiveSectionIdx = SectionIdx;
 
-	RefreshSequencerTracks();
+	SetDefaultMode();
 }
 
 void FContextualAnimViewModel::SetActiveAnimSetForSection(int32 SectionIdx, int32 AnimSetIdx)
@@ -125,7 +125,7 @@ void FContextualAnimViewModel::SetActiveAnimSetForSection(int32 SectionIdx, int3
 	int32& ActiveSetIdx = ActiveAnimSetMap.FindOrAdd(SectionIdx);
 	ActiveSetIdx = AnimSetIdx;
 
-	RefreshSequencerTracks();
+	SetDefaultMode();
 }
 
 AActor* FContextualAnimViewModel::SpawnPreviewActor(const FContextualAnimTrack& AnimTrack)
@@ -212,56 +212,36 @@ AActor* FContextualAnimViewModel::SpawnPreviewActor(const FContextualAnimTrack&
 	}
 }
 
-void FContextualAnimViewModel::RefreshSequencerTracks()
+void FContextualAnimViewModel::ResetTimeline()
 {
-	// Remove movie scene tracks and destroy existing actors (if any)
-
 	for (int32 MasterTrackIdx = MovieScene->GetMasterTracks().Num() - 1; MasterTrackIdx >= 0; MasterTrackIdx--)
 	{
 		UMovieSceneTrack& MasterTrack = *MovieScene->GetMasterTracks()[MasterTrackIdx];
-		for (const UMovieSceneSection* Section : MasterTrack.GetAllSections())
-		{
-			const UContextualAnimMovieSceneSection* ContextualAnimSection = Cast<const UContextualAnimMovieSceneSection>(Section);
-			check(ContextualAnimSection);
-
-			if(ContextualAnimSection->GetAnimTrack().Animation)
-			{
-				ContextualAnimSection->GetAnimTrack().Animation->UnregisterOnNotifyChanged(this);
-			}
-		}
-
 		MovieScene->RemoveMasterTrack(MasterTrack);
 	}
 
-	if (SceneInstance.IsValid())
-	{
-		SceneInstance->Stop();
-	}
+	Sequencer->NotifyMovieSceneDataChanged(EMovieSceneDataChangeType::MovieSceneStructureItemsChanged);
+}
 
-	if(StartSceneParams.RoleToActorMap.Num() > 0)
+void FContextualAnimViewModel::SetDefaultMode()
+{
+	if (EditingAnimation.IsValid())
 	{
-		for (auto& MapEntry : StartSceneParams.RoleToActorMap)
-		{			
-			if (AActor* Actor = MapEntry.Value.GetActor())
-			{
-				Actor->Destroy();
-			}
-		}
+		EditingAnimation->UnregisterOnNotifyChanged(this);
+		EditingAnimation.Reset();
 	}
 
-	Sequencer->NotifyMovieSceneDataChanged(EMovieSceneDataChangeType::MovieSceneStructureItemsChanged);
+	TimelineMode = ETimelineMode::Default;
 
-	if(!SceneAsset->Sections.IsValidIndex(ActiveSectionIdx))
+	ResetTimeline();
+
+	if (!SceneAsset->Sections.IsValidIndex(ActiveSectionIdx))
 	{
 		return;
 	}
 
-	StartSceneParams.Reset();
-	StartSceneParams.SectionIdx = ActiveSectionIdx;
-	StartSceneParams.AnimSetIdx = ActiveAnimSetMap.FindOrAdd(ActiveSectionIdx);
-
 	TArray<FName> Roles = SceneAsset->GetRoles();
-	for(FName Role : Roles)
+	for (FName Role : Roles)
 	{
 		UContextualAnimMovieSceneTrack* MovieSceneAnimTrack = MovieSceneSequence->GetMovieScene()->AddMasterTrack<UContextualAnimMovieSceneTrack>();
 		check(MovieSceneAnimTrack);
@@ -281,13 +261,7 @@ void FContextualAnimViewModel::RefreshSequencerTracks()
 			if (Animation)
 			{
 				UContextualAnimMovieSceneTrack* MovieSceneTrack = FindMasterTrackByRole(AnimTrack.Role);
-				if (MovieSceneTrack == nullptr)
-				{
-					UE_LOG(LogContextualAnim, Warning, TEXT("FContextualAnimViewModel::RefreshSequencerTracks. Can't find MovieSceneTrack for %s. Role: %s SectionIdx: %d AnimIndex: %d"),
-						*GetNameSafe(Animation), *AnimTrack.Role.ToString(), ActiveSectionIdx, AnimSetIdx);
-
-					continue;
-				}
+				check(MovieSceneTrack)
 
 				UContextualAnimMovieSceneSection* NewSection = NewObject<UContextualAnimMovieSceneSection>(MovieSceneTrack, UContextualAnimMovieSceneSection::StaticClass(), NAME_None, RF_Transactional);
 				check(NewSection);
@@ -306,23 +280,112 @@ void FContextualAnimViewModel::RefreshSequencerTracks()
 				MovieSceneTrack->AddSection(*NewSection);
 				MovieSceneTrack->SetTrackRowDisplayName(FText::FromString(FString::Printf(TEXT("%d"), AnimSetIdx)), AnimSetIdx);
 			}
+		}
+	}
+
+	Sequencer->NotifyMovieSceneDataChanged(EMovieSceneDataChangeType::MovieSceneStructureItemsChanged);
+
+	RefreshPreviewScene();
+}
+
+void FContextualAnimViewModel::SetNotifiesMode(const FContextualAnimTrack& AnimTrack)
+{
+	ResetTimeline();
+
+	TimelineMode = ETimelineMode::Notifies;
+
+	if(EditingAnimation.IsValid())
+	{
+		EditingAnimation->UnregisterOnNotifyChanged(this);
+	}
+
+	EditingAnimation = AnimTrack.Animation;
+
+	// Add Animation Track
+	{
+		UContextualAnimMovieSceneTrack* MovieSceneAnimTrack = MovieSceneSequence->GetMovieScene()->AddMasterTrack<UContextualAnimMovieSceneTrack>();
+		check(MovieSceneAnimTrack);
+
+		MovieSceneAnimTrack->SetDisplayName(FText::FromString(GetNameSafe(EditingAnimation.Get())));
+
+		UContextualAnimMovieSceneSection* NewSection = NewObject<UContextualAnimMovieSceneSection>(MovieSceneAnimTrack, UContextualAnimMovieSceneSection::StaticClass(), NAME_None, RF_Transactional);
+		check(NewSection);
+
+		NewSection->Initialize(AnimTrack.SectionIdx, AnimTrack.AnimSetIdx, AnimTrack.AnimTrackIdx);
 
-			if (!StartSceneParams.RoleToActorMap.Contains(AnimTrack.Role))
+		FFrameRate TickResolution = MovieSceneSequence->GetMovieScene()->GetTickResolution();
+		FFrameNumber StartFrame(0);
+		FFrameNumber EndFrame = (EditingAnimation->GetPlayLength() * TickResolution).RoundToFrame();
+		NewSection->SetRange(TRange<FFrameNumber>::Exclusive(StartFrame, EndFrame));
+
+		MovieSceneAnimTrack->AddSection(*NewSection);
+	}
+
+	// Add Notify Tracks
+	{
+		for (const FAnimNotifyTrack& NotifyTrack : EditingAnimation->AnimNotifyTracks)
+		{
+			UContextualAnimMovieSceneNotifyTrack* Track = MovieSceneSequence->GetMovieScene()->AddMasterTrack<UContextualAnimMovieSceneNotifyTrack>();
+			check(Track);
+
+			Track->Initialize(*EditingAnimation, NotifyTrack);
+		}
+
+		// Listen for when the notifies in the animation changes, so we can refresh the notify sections here
+		AnimTrack.Animation->RegisterOnNotifyChanged(UAnimSequenceBase::FOnNotifyChanged::CreateSP(this, &FContextualAnimViewModel::OnAnimNotifyChanged, EditingAnimation.Get()));
+	}
+
+	Sequencer->NotifyMovieSceneDataChanged(EMovieSceneDataChangeType::MovieSceneStructureItemsChanged);
+}
+
+void FContextualAnimViewModel::RefreshPreviewScene()
+{
+	if (SceneInstance.IsValid())
+	{
+		SceneInstance->Stop();
+	}
+
+	if (StartSceneParams.RoleToActorMap.Num() > 0)
+	{
+		for (auto& MapEntry : StartSceneParams.RoleToActorMap)
+		{
+			if (AActor* Actor = MapEntry.Value.GetActor())
 			{
-				if (AActor* PreviewActor = SpawnPreviewActor(AnimTrack))
-				{
-					StartSceneParams.RoleToActorMap.Add(AnimTrack.Role, PreviewActor);
-				}
+				Actor->Destroy();
 			}
 		}
 	}
 
-	Sequencer->NotifyMovieSceneDataChanged(EMovieSceneDataChangeType::MovieSceneStructureItemsChanged);
+	StartSceneParams.Reset();
+
+	if (!SceneAsset->Sections.IsValidIndex(ActiveSectionIdx))
+	{
+		return;
+	}
+
+	const FContextualAnimSceneSection& ContextualAnimSection = SceneAsset->Sections[ActiveSectionIdx];
+	const int32 ActiveAnimSetIdx = ActiveAnimSetMap.FindOrAdd(ActiveSectionIdx);
+	if (!ContextualAnimSection.AnimSets.IsValidIndex(ActiveAnimSetIdx))
+	{
+		return;
+	}
+	
+	const FContextualAnimSet& AnimSet = ContextualAnimSection.AnimSets[ActiveAnimSetIdx];
+	for (const FContextualAnimTrack& AnimTrack : AnimSet.Tracks)
+	{
+		if (AActor* PreviewActor = SpawnPreviewActor(AnimTrack))
+		{
+			StartSceneParams.RoleToActorMap.Add(AnimTrack.Role, PreviewActor);
+		}
+	}
+
+	StartSceneParams.SectionIdx = ActiveSectionIdx;
+	StartSceneParams.AnimSetIdx = ActiveAnimSetIdx;
 
 	SceneInstance = ContextualAnimManager->ForceStartScene(*GetSceneAsset(), StartSceneParams);
 
 	// Disable auto blend out
-	if(SceneInstance.IsValid())
+	if (SceneInstance.IsValid())
 	{
 		for (auto& Binding : SceneInstance->GetBindings())
 		{
@@ -433,9 +496,8 @@ void FContextualAnimViewModel::ToggleSimulateMode()
 				}
 			}
 		}
-		
 
-		RefreshSequencerTracks();
+		SetDefaultMode();
 	}
 };
 
@@ -605,7 +667,14 @@ void FContextualAnimViewModel::OnAnimNotifyChanged(UAnimSequenceBase* Animation)
 
 	UE_LOG(LogContextualAnim, Log, TEXT("FContextualAnimViewModel::OnAnimNotifyChanged Anim: %s. Refreshing Sequencer Tracks"), *GetNameSafe(Animation));
 
-	RefreshSequencerTracks();
+	// Refresh notifies view if active
+	if(TimelineMode == ETimelineMode::Notifies && EditingAnimation.Get() == Animation)
+	{
+		if(const FContextualAnimTrack* AnimTrack = GetSceneAsset()->FindAnimTrackByAnimation(Animation))
+		{
+			SetNotifiesMode(*AnimTrack);
+		}
+	}
 }
 
 void FContextualAnimViewModel::AnimationModified(UAnimSequenceBase& Animation)
diff --git a/Engine/Plugins/Experimental/Animation/ContextualAnimation/Source/ContextualAnimationEditor/Public/ContextualAnimMovieSceneNotifyTrackEditor.h b/Engine/Plugins/Experimental/Animation/ContextualAnimation/Source/ContextualAnimationEditor/Public/ContextualAnimMovieSceneNotifyTrackEditor.h
--- a/Engine/Plugins/Experimental/Animation/ContextualAnimation/Source/ContextualAnimationEditor/Public/ContextualAnimMovieSceneNotifyTrackEditor.h
+++ b/Engine/Plugins/Experimental/Animation/ContextualAnimation/Source/ContextualAnimationEditor/Public/ContextualAnimMovieSceneNotifyTrackEditor.h
@@ -15,6 +15,7 @@ class UContextualAnimMovieSceneNotifySection;
 class UContextualAnimMovieSceneNotifyTrack;
 class UContextualAnimMovieSceneSequence;
 class UContextualAnimNewIKTargetParams;
+class FToolBarBuilder;
 
 /** Handles section drawing and manipulation of a MovieSceneNotifyTrack */
 class FContextualAnimMovieSceneNotifyTrackEditor : public FMovieSceneTrackEditor, public FGCObject
@@ -28,6 +29,7 @@ public:
 	virtual FString GetReferencerName() const override { return TEXT("FContextualAnimMovieSceneNotifyTrackEditor"); }
 
 	// ~FMovieSceneTrackEditor Interface
+	virtual void OnInitialize() override;
 	virtual TSharedRef<ISequencerSection> MakeSectionInterface(UMovieSceneSection& SectionObject, UMovieSceneTrack& Track, FGuid ObjectBinding) override;
 	virtual bool SupportsType(TSubclassOf<UMovieSceneTrack> TrackClass) const override;
 	virtual bool SupportsSequence(UMovieSceneSequence* InSequence) const override;
@@ -51,9 +53,11 @@ private:
 
 	void BuildNewIKTargetWidget(FMenuBuilder& MenuBuilder, UContextualAnimMovieSceneNotifyTrack* Track, int32 RowIndex);
 
-	void AddNewNotifyTrack(TArray<FGuid> ObjectBindings);
+	void AddNewNotifyTrack();
 
 	void FillNewNotifyStateMenu(FMenuBuilder& MenuBuilder, bool bIsReplaceWithMenu, UContextualAnimMovieSceneNotifyTrack* Track, int32 RowIndex);
+
+	void CustomizeToolBar(FToolBarBuilder& ToolBarBuilder);
 };
 
 // FContextualAnimNotifySection
diff --git a/Engine/Plugins/Experimental/Animation/ContextualAnimation/Source/ContextualAnimationEditor/Public/ContextualAnimViewModel.h b/Engine/Plugins/Experimental/Animation/ContextualAnimation/Source/ContextualAnimationEditor/Public/ContextualAnimViewModel.h
--- a/Engine/Plugins/Experimental/Animation/ContextualAnimation/Source/ContextualAnimationEditor/Public/ContextualAnimViewModel.h
+++ b/Engine/Plugins/Experimental/Animation/ContextualAnimation/Source/ContextualAnimationEditor/Public/ContextualAnimViewModel.h
@@ -33,6 +33,12 @@ class FContextualAnimViewModel : public TSharedFromThis<FContextualAnimViewModel
 {
 public:
 
+	enum class ETimelineMode : uint8
+	{
+		Default,
+		Notifies
+	};
+
 	FContextualAnimViewModel();
 	virtual ~FContextualAnimViewModel();
 
@@ -42,7 +48,12 @@ public:
 
 	void Initialize(UContextualAnimSceneAsset* InSceneAsset, const TSharedRef<FContextualAnimPreviewScene>& InPreviewScene);
 
-	void RefreshSequencerTracks();
+	void SetDefaultMode();
+	void SetNotifiesMode(const FContextualAnimTrack& AnimTrack);
+	ETimelineMode GetTimelineMode() const { return TimelineMode; }
+	UAnimSequenceBase* GetEditingAnimation() const { return EditingAnimation.Get(); }
+	void RefreshPreviewScene();
+	void ResetTimeline();
 
 	TSharedPtr<ISequencer> GetSequencer();
 	TSharedPtr<FContextualAnimPreviewScene> GetPreviewScene();
@@ -148,6 +159,10 @@ private:
 
 	FContextualAnimStartSceneParams StartSceneParams;
 
+	ETimelineMode TimelineMode = ETimelineMode::Default;
+
+	TWeakObjectPtr<UAnimSequenceBase> EditingAnimation = nullptr;
+
 	AActor* SpawnPreviewActor(const FContextualAnimTrack& AnimTrack);
 
 	UWorld* GetWorld() const;
