Contextual Anim Editor: Added support for adding and editting simple notifies

diff --git a/Engine/Plugins/Experimental/Animation/ContextualAnimation/Source/ContextualAnimationEditor/Private/ContextualAnimMovieSceneNotifySection.cpp b/Engine/Plugins/Experimental/Animation/ContextualAnimation/Source/ContextualAnimationEditor/Private/ContextualAnimMovieSceneNotifySection.cpp
--- a/Engine/Plugins/Experimental/Animation/ContextualAnimation/Source/ContextualAnimationEditor/Private/ContextualAnimMovieSceneNotifySection.cpp
+++ b/Engine/Plugins/Experimental/Animation/ContextualAnimation/Source/ContextualAnimationEditor/Private/ContextualAnimMovieSceneNotifySection.cpp
@@ -12,9 +12,24 @@ UContextualAnimMovieSceneNotifyTrack* UContextualAnimMovieSceneNotifySection::Ge
 
 void UContextualAnimMovieSceneNotifySection::Initialize(const FAnimNotifyEvent& NotifyEvent)
 {
-	FFrameRate TickResolution = GetOwnerTrack()->GetTypedOuter<UMovieScene>()->GetTickResolution();
-	FFrameNumber StartFrame = (NotifyEvent.GetTriggerTime() * TickResolution).RoundToFrame();
-	FFrameNumber EndFrame = (NotifyEvent.GetEndTriggerTime() * TickResolution).RoundToFrame();
+	const FFrameRate TickResolution = GetOwnerTrack()->GetTypedOuter<UMovieScene>()->GetTickResolution();
+	const FFrameNumber StartFrame = (NotifyEvent.GetTriggerTime() * TickResolution).RoundToFrame();
+
+	FFrameNumber EndFrame;
+	if (NotifyEvent.NotifyStateClass)
+	{
+		EndFrame = (NotifyEvent.GetEndTriggerTime() * TickResolution).RoundToFrame();
+
+	}
+	else if (NotifyEvent.Notify)
+	{
+		// Sequencer Panel doesn't seem to have a way to represent single key events the same way we represent AnimNotify in the Animation Editor.
+		// So, for now we represent single notifies as sections with a small fixed interval and mark them as non resizable (see FContextualAnimNotifySection::SectionIsResizable). 
+		// This may need to be revisited in the future.
+
+		EndFrame = ((NotifyEvent.GetTriggerTime() + 1.f / 30.f) * TickResolution).RoundToFrame();
+	}
+
 	SetRange(TRange<FFrameNumber>::Exclusive(StartFrame, EndFrame));
 
 	AnimNotifyEventGuid = NotifyEvent.Guid;
diff --git a/Engine/Plugins/Experimental/Animation/ContextualAnimation/Source/ContextualAnimationEditor/Private/ContextualAnimMovieSceneNotifyTrack.cpp b/Engine/Plugins/Experimental/Animation/ContextualAnimation/Source/ContextualAnimationEditor/Private/ContextualAnimMovieSceneNotifyTrack.cpp
--- a/Engine/Plugins/Experimental/Animation/ContextualAnimation/Source/ContextualAnimationEditor/Private/ContextualAnimMovieSceneNotifyTrack.cpp
+++ b/Engine/Plugins/Experimental/Animation/ContextualAnimation/Source/ContextualAnimationEditor/Private/ContextualAnimMovieSceneNotifyTrack.cpp
@@ -24,17 +24,14 @@ void UContextualAnimMovieSceneNotifyTrack::Initialize(UAnimSequenceBase& Animati
 
 	for (const FAnimNotifyEvent* NotifyEvent : NotifyTrack.Notifies)
 	{
-		if (UAnimNotifyState* Notify = Cast<UAnimNotifyState>(NotifyEvent->NotifyStateClass))
-		{
-			// Create new movie scene section
-			UContextualAnimMovieSceneNotifySection* Section = CastChecked<UContextualAnimMovieSceneNotifySection>(CreateNewSection());
+		// Create new movie scene section
+		UContextualAnimMovieSceneNotifySection* Section = CastChecked<UContextualAnimMovieSceneNotifySection>(CreateNewSection());
 
-			// Set range and cache guid
-			Section->Initialize(*NotifyEvent);
+		// Set range and cache guid
+		Section->Initialize(*NotifyEvent);
 
-			// Add section to the track
-			AddSection(*Section);
-		}
+		// Add section to the track
+		AddSection(*Section);
 	}
 }
 
diff --git a/Engine/Plugins/Experimental/Animation/ContextualAnimation/Source/ContextualAnimationEditor/Private/ContextualAnimMovieSceneNotifyTrackEditor.cpp b/Engine/Plugins/Experimental/Animation/ContextualAnimation/Source/ContextualAnimationEditor/Private/ContextualAnimMovieSceneNotifyTrackEditor.cpp
--- a/Engine/Plugins/Experimental/Animation/ContextualAnimation/Source/ContextualAnimationEditor/Private/ContextualAnimMovieSceneNotifyTrackEditor.cpp
+++ b/Engine/Plugins/Experimental/Animation/ContextualAnimation/Source/ContextualAnimationEditor/Private/ContextualAnimMovieSceneNotifyTrackEditor.cpp
@@ -24,6 +24,7 @@
 #include "ContextualAnimEditorTypes.h"
 #include "ISequencerModule.h"
 #include "Toolkits/AssetEditorToolkit.h"
+#include "Fonts/FontMeasure.h"
 
 #define LOCTEXT_NAMESPACE "FContextualAnimMovieSceneNotifyTrackEditor"
 
@@ -59,7 +60,7 @@ TSharedRef<ISequencerTrackEditor> FContextualAnimMovieSceneNotifyTrackEditor::Cr
 TSharedRef<ISequencerSection> FContextualAnimMovieSceneNotifyTrackEditor::MakeSectionInterface(UMovieSceneSection& SectionObject, UMovieSceneTrack& Track, FGuid ObjectBinding)
 {
 	checkf(SectionObject.GetClass()->IsChildOf<UContextualAnimMovieSceneNotifySection>(), TEXT("Unsupported section."));
-	return MakeShareable(new FContextualAnimNotifySection(SectionObject));
+	return MakeShared<FContextualAnimNotifySection>(SectionObject);
 }
 
 void FContextualAnimMovieSceneNotifyTrackEditor::OnInitialize()
@@ -102,7 +103,7 @@ void FContextualAnimMovieSceneNotifyTrackEditor::CustomizeToolBar(FToolBarBuilde
 	ToolBarBuilder.EndSection();
 }
 
-void FContextualAnimMovieSceneNotifyTrackEditor::FillNewNotifyStateMenu(FMenuBuilder& MenuBuilder, bool bIsReplaceWithMenu, UContextualAnimMovieSceneNotifyTrack* Track, int32 RowIndex)
+void FContextualAnimMovieSceneNotifyTrackEditor::FillNewNotifyMenu(FMenuBuilder& MenuBuilder, bool bIsAnimNotifyState, UContextualAnimMovieSceneNotifyTrack* Track, int32 RowIndex)
 {
 	// MenuBuilder always has a search widget added to it by default, hence if larger then 1 then something else has been added to it
 	if (MenuBuilder.GetMultiBox()->GetBlocks().Num() > 1)
@@ -110,22 +111,15 @@ void FContextualAnimMovieSceneNotifyTrackEditor::FillNewNotifyStateMenu(FMenuBui
 		MenuBuilder.AddMenuSeparator();
 	}
 
-	FOnClassPicked OnPicked(FOnClassPicked::CreateLambda([this, bIsReplaceWithMenu, Track, RowIndex](UClass* InClass)
+	FOnClassPicked OnPicked(FOnClassPicked::CreateLambda([this, Track, RowIndex](UClass* InClass)
 		{
 			//FSlateApplication::Get().DismissAllMenus();
 
-			if (bIsReplaceWithMenu)
-			{
-				//@TODO: Implement
-			}
-			else
-			{
-				CreateNewSection(Track, RowIndex, InClass);
-			}
+			CreateNewSection(Track, RowIndex, InClass);
 		}
 	));
 
-	TSharedRef<SWidget> AnimNotifyStateClassPicker = PersonaUtils::MakeAnimNotifyStatePicker(&Track->GetAnimation(), OnPicked);
+	TSharedRef<SWidget> AnimNotifyStateClassPicker = bIsAnimNotifyState ? PersonaUtils::MakeAnimNotifyStatePicker(&Track->GetAnimation(), OnPicked) : PersonaUtils::MakeAnimNotifyPicker(&Track->GetAnimation(), OnPicked);
 	MenuBuilder.AddWidget(AnimNotifyStateClassPicker, FText(), true, false);
 }
 
@@ -166,18 +160,22 @@ TSharedPtr<SWidget> FContextualAnimMovieSceneNotifyTrackEditor::BuildOutlinerEdi
 		UContextualAnimMovieSceneNotifyTrack* TrackPtr = WeakTrack.Get();
 		if (TrackPtr)
 		{
-			//@TODO: Add shortcut for Motion Warping window too
+			//@TODO: Add shortcut for Motion Warping window too and group these options in two sections
 
 			MenuBuilder.AddSubMenu(
 				LOCTEXT("AddIKWindow", "Add IK Window"),
 				LOCTEXT("AddIKWindowTooltip", "Adds new IK Window"),
 				FNewMenuDelegate::CreateRaw(this, &FContextualAnimMovieSceneNotifyTrackEditor::BuildNewIKTargetSubMenu, TrackPtr, RowIndex));
 
+			MenuBuilder.AddSubMenu(
+				LOCTEXT("AddNotify", "Add Notify"),
+				LOCTEXT("AddNotifyToolTip", "Adds new AnimNotify"),
+				FNewMenuDelegate::CreateRaw(this, &FContextualAnimMovieSceneNotifyTrackEditor::FillNewNotifyMenu, false /*bIsAnimNotifyState*/, TrackPtr, RowIndex));
 
 			MenuBuilder.AddSubMenu(
 				LOCTEXT("AddNotifyState", "Add Notify State"),
 				LOCTEXT("AddNotifyStateToolTip", "Adds new AnimNotifyState"),
-				FNewMenuDelegate::CreateRaw(this, &FContextualAnimMovieSceneNotifyTrackEditor::FillNewNotifyStateMenu, false, TrackPtr, RowIndex));
+				FNewMenuDelegate::CreateRaw(this, &FContextualAnimMovieSceneNotifyTrackEditor::FillNewNotifyMenu, true /*bIsAnimNotifyState*/, TrackPtr, RowIndex));
 		}
 		else
 		{
@@ -452,37 +450,127 @@ UContextualAnimMovieSceneNotifySection* FContextualAnimMovieSceneNotifyTrackEdit
 // FContextualAnimNotifySection
 ////////////////////////////////////////////////////////////////////////////////////////////////
 
-FContextualAnimNotifySection::FContextualAnimNotifySection(UMovieSceneSection& InSection)
-	: Section(*CastChecked<UContextualAnimMovieSceneNotifySection>(&InSection))
+UContextualAnimMovieSceneNotifySection* FContextualAnimNotifySection::GetNotifySection() const
 {
+	if (UContextualAnimMovieSceneNotifySection* NotifySection = Cast<UContextualAnimMovieSceneNotifySection>(WeakSection.Get()))
+	{
+		return NotifySection;
+	}
+
+	return nullptr;
 }
 
-UMovieSceneSection* FContextualAnimNotifySection::GetSectionObject()
+void FContextualAnimNotifySection::PaintNotifyName(FSequencerSectionPainter& Painter, int32 LayerId, const FString& InEventString, float PixelPos, bool bIsEventValid)
 {
-	return &Section;
+	static const int32   FontSize = 10;
+	static const float   BoxOffsetPx = 10.f;
+	static const TCHAR* WarningString = TEXT("\xf071");
+
+	const FSlateFontInfo FontAwesomeFont = FAppStyle::Get().GetFontStyle("FontAwesome.10");
+	const FSlateFontInfo SmallLayoutFont = FCoreStyle::GetDefaultFontStyle("Bold", 10);
+	const FLinearColor   DrawColor = FAppStyle::GetSlateColor("SelectionColor").GetColor(FWidgetStyle());
+
+	TSharedRef<FSlateFontMeasure> FontMeasureService = FSlateApplication::Get().GetRenderer()->GetFontMeasureService();
+
+	// Setup the warning size. Static since it won't ever change
+	static FVector2D WarningSize = FontMeasureService->Measure(WarningString, FontAwesomeFont);
+	const  FMargin   WarningPadding = (bIsEventValid || InEventString.Len() == 0) ? FMargin(0.f) : FMargin(0.f, 0.f, 4.f, 0.f);
+	const  FMargin   BoxPadding = FMargin(4.0f, 2.0f);
+
+	const FVector2D  TextSize = FontMeasureService->Measure(InEventString, SmallLayoutFont);
+	const FVector2D  IconSize = bIsEventValid ? FVector2D::ZeroVector : WarningSize;
+	const FVector2D  PaddedIconSize = IconSize + WarningPadding.GetDesiredSize();
+	const FVector2D  BoxSize = FVector2D(TextSize.X + PaddedIconSize.X, FMath::Max(TextSize.Y, PaddedIconSize.Y)) + BoxPadding.GetDesiredSize();
+
+	FVector2D BoxOffset = FVector2D(PixelPos + BoxOffsetPx, Painter.SectionGeometry.Size.Y * .5f - BoxSize.Y * .5f);
+	FVector2D IconOffset = FVector2D(BoxPadding.Left, BoxSize.Y * .5f - IconSize.Y * .5f);
+	FVector2D TextOffset = FVector2D(IconOffset.X + PaddedIconSize.X, BoxSize.Y * .5f - TextSize.Y * .5f);
+
+	// Draw the background box
+	FSlateDrawElement::MakeBox(
+		Painter.DrawElements,
+		LayerId + 1,
+		Painter.SectionGeometry.ToPaintGeometry(BoxOffset, BoxSize),
+		FAppStyle::GetBrush("WhiteBrush"),
+		ESlateDrawEffect::None,
+		FLinearColor::Black.CopyWithNewOpacity(0.5f)
+	);
+
+	if (!bIsEventValid)
+	{
+		// Draw a warning icon for unbound repeaters
+		FSlateDrawElement::MakeText(
+			Painter.DrawElements,
+			LayerId + 2,
+			Painter.SectionGeometry.ToPaintGeometry(BoxOffset + IconOffset, IconSize),
+			WarningString,
+			FontAwesomeFont,
+			Painter.bParentEnabled ? ESlateDrawEffect::None : ESlateDrawEffect::DisabledEffect,
+			FAppStyle::GetWidgetStyle<FTextBlockStyle>("Log.Warning").ColorAndOpacity.GetSpecifiedColor()
+		);
+	}
+
+	FSlateDrawElement::MakeText(
+		Painter.DrawElements,
+		LayerId + 2,
+		Painter.SectionGeometry.ToPaintGeometry(BoxOffset + TextOffset, TextSize),
+		InEventString,
+		SmallLayoutFont,
+		Painter.bParentEnabled ? ESlateDrawEffect::None : ESlateDrawEffect::DisabledEffect,
+		DrawColor
+	);
 }
 
 int32 FContextualAnimNotifySection::OnPaintSection(FSequencerSectionPainter& InPainter) const
 {
-	const UAnimNotifyState* AnimNotifyState = Section.GetAnimNotifyState();
-	if (AnimNotifyState)
+	int32 LayerId = InPainter.PaintSectionBackground();
+
+	// Pop current clip state since we want the notify name to 'overflow' the geometry of the section
+	TOptional<FSlateClippingState> PreviousClipState = InPainter.DrawElements.GetClippingState();
+	InPainter.DrawElements.PopClip();
+
+	// Draw notify name
+	FContextualAnimNotifySection::PaintNotifyName(InPainter, LayerId, GetNotifyName(), 0, true);
+
+	// Restore clip state
+	if (PreviousClipState.IsSet())
+	{
+		InPainter.DrawElements.GetClippingManager().PushClippingState(PreviousClipState.GetValue());
+	}
+
+	return LayerId;
+}
+
+bool FContextualAnimNotifySection::SectionIsResizable() const
+{
+	if (UContextualAnimMovieSceneNotifySection* NotifySection = GetNotifySection())
 	{
-		//@TODO: Use AnimNotifyState->NotifyColor
-		return InPainter.PaintSectionBackground();
+		const FAnimNotifyEvent* AnimNotifyEvent = NotifySection->GetAnimNotifyEvent();
+		return (AnimNotifyEvent && AnimNotifyEvent->NotifyStateClass);
 	}
 
-	return InPainter.PaintSectionBackground(FLinearColor::Red);
+	return false;
 }
 
-FText FContextualAnimNotifySection::GetSectionTitle() const
+FString FContextualAnimNotifySection::GetNotifyName() const
 {
-	const UAnimNotifyState* AnimNotifyState = Section.GetAnimNotifyState();
-	if (AnimNotifyState)
+	if (const UContextualAnimMovieSceneNotifySection* NotifySection = GetNotifySection())
 	{
-		return FText::FromString(AnimNotifyState->GetNotifyName());
+		const FAnimNotifyEvent* AnimNotifyEvent = NotifySection->GetAnimNotifyEvent();
+		if (AnimNotifyEvent)
+		{
+			if (AnimNotifyEvent->NotifyStateClass)
+			{
+				return AnimNotifyEvent->NotifyStateClass->GetNotifyName();
+			}
+			else if (AnimNotifyEvent->Notify)
+			{
+				return AnimNotifyEvent->Notify->GetNotifyName();
+			}
+		}
 	}
 
-	return LOCTEXT("InvalidAnimNotifyState", "Invalid AnimNotifyState");
+	return FString(TEXT("Invalid"));
 }
 
 #undef LOCTEXT_NAMESPACE
\ No newline at end of file
diff --git a/Engine/Plugins/Experimental/Animation/ContextualAnimation/Source/ContextualAnimationEditor/Public/ContextualAnimMovieSceneNotifyTrackEditor.h b/Engine/Plugins/Experimental/Animation/ContextualAnimation/Source/ContextualAnimationEditor/Public/ContextualAnimMovieSceneNotifyTrackEditor.h
--- a/Engine/Plugins/Experimental/Animation/ContextualAnimation/Source/ContextualAnimationEditor/Public/ContextualAnimMovieSceneNotifyTrackEditor.h
+++ b/Engine/Plugins/Experimental/Animation/ContextualAnimation/Source/ContextualAnimationEditor/Public/ContextualAnimMovieSceneNotifyTrackEditor.h
@@ -55,7 +55,7 @@ private:
 
 	void AddNewNotifyTrack();
 
-	void FillNewNotifyStateMenu(FMenuBuilder& MenuBuilder, bool bIsReplaceWithMenu, UContextualAnimMovieSceneNotifyTrack* Track, int32 RowIndex);
+	void FillNewNotifyMenu(FMenuBuilder& MenuBuilder, bool bIsAnimNotifyState, UContextualAnimMovieSceneNotifyTrack* Track, int32 RowIndex);
 
 	void CustomizeToolBar(FToolBarBuilder& ToolBarBuilder);
 };
@@ -64,19 +64,19 @@ private:
 ////////////////////////////////////////////////////////////////////////////////////////////////
 
 /** UI portion of a NotifySection in a NotifyTrack */
-class FContextualAnimNotifySection : public ISequencerSection
+class FContextualAnimNotifySection : public FSequencerSection
 {
 public:
 	
-	FContextualAnimNotifySection(UMovieSceneSection& InSection);
-	virtual ~FContextualAnimNotifySection() { }
+	FContextualAnimNotifySection(UMovieSceneSection& InSectionObject)
+		: FSequencerSection(InSectionObject)
+	{}
 
-	virtual UMovieSceneSection* GetSectionObject() override;
-	virtual int32 OnPaintSection(FSequencerSectionPainter& InPainter) const override;
-	virtual FText GetSectionTitle() const override;
+	virtual int32 OnPaintSection(FSequencerSectionPainter& Painter) const override;
+	virtual bool SectionIsResizable() const override;
 
-private:
-	
-	/** The section we are visualizing */
-	UContextualAnimMovieSceneNotifySection& Section;
+	UContextualAnimMovieSceneNotifySection* GetNotifySection() const;
+	FString GetNotifyName() const;
+
+	static void PaintNotifyName(FSequencerSectionPainter& Painter, int32 LayerId, const FString& InEventString, float PixelPos, bool bIsEventValid);
 };
\ No newline at end of file
