Contextual Anim Editor: Fix for text conflict from LOCTEXT macro

diff --git a/Engine/Plugins/Experimental/Animation/ContextualAnimation/Source/ContextualAnimationEditor/Private/ContextualAnimMovieSceneNotifyTrackEditor.cpp b/Engine/Plugins/Experimental/Animation/ContextualAnimation/Source/ContextualAnimationEditor/Private/ContextualAnimMovieSceneNotifyTrackEditor.cpp
--- a/Engine/Plugins/Experimental/Animation/ContextualAnimation/Source/ContextualAnimationEditor/Private/ContextualAnimMovieSceneNotifyTrackEditor.cpp
+++ b/Engine/Plugins/Experimental/Animation/ContextualAnimation/Source/ContextualAnimationEditor/Private/ContextualAnimMovieSceneNotifyTrackEditor.cpp
@@ -190,7 +190,7 @@ TSharedPtr<SWidget> FContextualAnimMovieSceneNotifyTrackEditor::BuildOutlinerEdi
 		.AutoWidth()
 		.VAlign(VAlign_Center)
 		[
-			FSequencerUtilities::MakeAddButton(LOCTEXT("AddNotify", "Notify"), FOnGetContent::CreateLambda(SubMenuCallback), Params.NodeIsHovered, GetSequencer())
+			FSequencerUtilities::MakeAddButton(LOCTEXT("Notify", "Notify"), FOnGetContent::CreateLambda(SubMenuCallback), Params.NodeIsHovered, GetSequencer())
 		];
 }
 
