Contextual Anim Editor: Fix for crash when removing section or animset from the details panel

diff --git a/Engine/Plugins/Experimental/Animation/ContextualAnimation/Source/ContextualAnimationEditor/Private/ContextualAnimAssetEditorToolkit.cpp b/Engine/Plugins/Experimental/Animation/ContextualAnimation/Source/ContextualAnimationEditor/Private/ContextualAnimAssetEditorToolkit.cpp
--- a/Engine/Plugins/Experimental/Animation/ContextualAnimation/Source/ContextualAnimationEditor/Private/ContextualAnimAssetEditorToolkit.cpp
+++ b/Engine/Plugins/Experimental/Animation/ContextualAnimation/Source/ContextualAnimationEditor/Private/ContextualAnimAssetEditorToolkit.cpp
@@ -455,11 +455,7 @@ TSharedRef<SDockTab> FContextualAnimAssetEditorToolkit::SpawnTab_PreviewSettings
 
 void FContextualAnimAssetEditorToolkit::OnFinishedChangingProperties(const FPropertyChangedEvent& PropertyChangedEvent)
 {
-	const FName PropertyName = (PropertyChangedEvent.Property != nullptr) ? PropertyChangedEvent.Property->GetFName() : NAME_None;
-	const FName MemberPropertyName = (PropertyChangedEvent.MemberProperty != nullptr) ? PropertyChangedEvent.MemberProperty->GetFName() : NAME_None;
-
-	UE_LOG(LogContextualAnim, Log, TEXT("FContextualAnimAssetEditorToolkit::OnFinishedChangingProperties MemberPropertyName: %s PropertyName: %s"), 
-		*MemberPropertyName.ToString(), *PropertyName.ToString());
+	ViewModel->OnFinishedChangingProperties(PropertyChangedEvent);
 }
 
 #undef LOCTEXT_NAMESPACE
diff --git a/Engine/Plugins/Experimental/Animation/ContextualAnimation/Source/ContextualAnimationEditor/Private/ContextualAnimViewModel.cpp b/Engine/Plugins/Experimental/Animation/ContextualAnimation/Source/ContextualAnimationEditor/Private/ContextualAnimViewModel.cpp
--- a/Engine/Plugins/Experimental/Animation/ContextualAnimation/Source/ContextualAnimationEditor/Private/ContextualAnimViewModel.cpp
+++ b/Engine/Plugins/Experimental/Animation/ContextualAnimation/Source/ContextualAnimationEditor/Private/ContextualAnimViewModel.cpp
@@ -237,6 +237,7 @@ void FContextualAnimViewModel::SetDefaultMode()
 
 	if (!SceneAsset->Sections.IsValidIndex(ActiveSectionIdx))
 	{
+		RefreshPreviewScene();
 		return;
 	}
 
@@ -342,7 +343,9 @@ void FContextualAnimViewModel::RefreshPreviewScene()
 {
 	if (SceneInstance.IsValid())
 	{
-		SceneInstance->Stop();
+		SceneInstance->GetBindings().Reset();
+		ContextualAnimManager->OnSceneInstanceEnded(SceneInstance.Get());
+		SceneInstance.Reset();
 	}
 
 	if (StartSceneParams.RoleToActorMap.Num() > 0)
@@ -878,4 +881,18 @@ FVector FContextualAnimViewModel::GetWidgetLocationFromSelection() const
 	}
 
 	return FVector::ZeroVector;
+}
+
+void FContextualAnimViewModel::OnFinishedChangingProperties(const FPropertyChangedEvent& PropertyChangedEvent)
+{
+	const FName PropertyName = (PropertyChangedEvent.Property != nullptr) ? PropertyChangedEvent.Property->GetFName() : NAME_None;
+	const FName MemberPropertyName = (PropertyChangedEvent.MemberProperty != nullptr) ? PropertyChangedEvent.MemberProperty->GetFName() : NAME_None;
+
+	// Refresh the time line if the Sections array or the AnimSets array inside a section changes
+	if (MemberPropertyName == GET_MEMBER_NAME_CHECKED(UContextualAnimSceneAsset, Sections) &&
+		(PropertyName == GET_MEMBER_NAME_CHECKED(UContextualAnimSceneAsset, Sections) ||
+			PropertyName == GET_MEMBER_NAME_CHECKED(FContextualAnimSceneSection, AnimSets)))
+	{
+		SetDefaultMode();
+	}
 }
\ No newline at end of file
diff --git a/Engine/Plugins/Experimental/Animation/ContextualAnimation/Source/ContextualAnimationEditor/Public/ContextualAnimViewModel.h b/Engine/Plugins/Experimental/Animation/ContextualAnimation/Source/ContextualAnimationEditor/Public/ContextualAnimViewModel.h
--- a/Engine/Plugins/Experimental/Animation/ContextualAnimation/Source/ContextualAnimationEditor/Public/ContextualAnimViewModel.h
+++ b/Engine/Plugins/Experimental/Animation/ContextualAnimation/Source/ContextualAnimationEditor/Public/ContextualAnimViewModel.h
@@ -46,6 +46,8 @@ public:
 	virtual void AddReferencedObjects(FReferenceCollector& Collector) override;
 	virtual FString GetReferencerName() const override { return TEXT("FContextualAnimViewModel"); }
 
+	void OnFinishedChangingProperties(const FPropertyChangedEvent& PropertyChangedEvent);
+
 	void Initialize(UContextualAnimSceneAsset* InSceneAsset, const TSharedRef<FContextualAnimPreviewScene>& InPreviewScene);
 
 	void SetDefaultMode();
