Fix non unity build errors
#test compile cpp

diff --git a/Engine/Plugins/Animation/IKRig/Source/IKRigEditor/Public/RigEditor/IKRigEditorController.h b/Engine/Plugins/Animation/IKRig/Source/IKRigEditor/Public/RigEditor/IKRigEditorController.h
--- a/Engine/Plugins/Animation/IKRig/Source/IKRigEditor/Public/RigEditor/IKRigEditorController.h
+++ b/Engine/Plugins/Animation/IKRig/Source/IKRigEditor/Public/RigEditor/IKRigEditorController.h
@@ -19,6 +19,7 @@ class SIKRigSkeleton;
 class UIKRigController;
 class FSolverStackElement;
 class FIKRigTreeElement;
+class UAnimInstance;
 class UDebugSkelMeshComponent;
 
 enum EIKRigSelectionType : int8
diff --git a/Engine/Plugins/Experimental/Animation/ContextualAnimation/Source/ContextualAnimationEditor/Private/ContextualAnimEditorTypes.cpp b/Engine/Plugins/Experimental/Animation/ContextualAnimation/Source/ContextualAnimationEditor/Private/ContextualAnimEditorTypes.cpp
--- a/Engine/Plugins/Experimental/Animation/ContextualAnimation/Source/ContextualAnimationEditor/Private/ContextualAnimEditorTypes.cpp
+++ b/Engine/Plugins/Experimental/Animation/ContextualAnimation/Source/ContextualAnimationEditor/Private/ContextualAnimEditorTypes.cpp
@@ -3,6 +3,7 @@
 #include "ContextualAnimEditorTypes.h"
 #include "ContextualAnimSceneAsset.h"
 #include "Animation/AnimMontage.h"
+#include "PropertyHandle.h"
 
 // UContextualAnimNewIKTargetParams
 ////////////////////////////////////////////////////////////////////////////////////////////////
diff --git a/Engine/Plugins/Experimental/Animation/ContextualAnimation/Source/ContextualAnimationEditor/Private/SContextualAnimAssetBrowser.cpp b/Engine/Plugins/Experimental/Animation/ContextualAnimation/Source/ContextualAnimationEditor/Private/SContextualAnimAssetBrowser.cpp
--- a/Engine/Plugins/Experimental/Animation/ContextualAnimation/Source/ContextualAnimationEditor/Private/SContextualAnimAssetBrowser.cpp
+++ b/Engine/Plugins/Experimental/Animation/ContextualAnimation/Source/ContextualAnimationEditor/Private/SContextualAnimAssetBrowser.cpp
@@ -5,6 +5,7 @@
 #include "IContentBrowserSingleton.h"
 #include "Animation/AnimSequence.h"
 #include "Animation/AnimMontage.h"
+#include "Components/VerticalBox.h"
 
 void SContextualAnimAssetBrowser::Construct(const FArguments& InArgs)
 {
diff --git a/Engine/Plugins/Experimental/Animation/PoseSearch/Source/Editor/Public/SPoseSearchDatabaseAssetList.h b/Engine/Plugins/Experimental/Animation/PoseSearch/Source/Editor/Public/SPoseSearchDatabaseAssetList.h
--- a/Engine/Plugins/Experimental/Animation/PoseSearch/Source/Editor/Public/SPoseSearchDatabaseAssetList.h
+++ b/Engine/Plugins/Experimental/Animation/PoseSearch/Source/Editor/Public/SPoseSearchDatabaseAssetList.h
@@ -7,6 +7,7 @@
 #include "Widgets/Views/STreeView.h"
 #include "EditorUndoClient.h"
 
+class FUICommandList;
 
 namespace UE::PoseSearch
 {
diff --git a/Engine/Plugins/Experimental/GameplayBehaviors/Source/GameplayBehaviorsModule/Private/GameplayBehavior_AnimationBased.cpp b/Engine/Plugins/Experimental/GameplayBehaviors/Source/GameplayBehaviorsModule/Private/GameplayBehavior_AnimationBased.cpp
--- a/Engine/Plugins/Experimental/GameplayBehaviors/Source/GameplayBehaviorsModule/Private/GameplayBehavior_AnimationBased.cpp
+++ b/Engine/Plugins/Experimental/GameplayBehaviors/Source/GameplayBehaviorsModule/Private/GameplayBehavior_AnimationBased.cpp
@@ -6,7 +6,7 @@
 #include "AbilitySystemComponent.h"
 #include "GameFramework/Actor.h"
 #include "Animation/AnimMontage.h"
-
+#include "Engine/World.h"
 
 //----------------------------------------------------------------------//
 // UGameplayBehavior_AnimationBased
diff --git a/Engine/Plugins/Importers/USDImporter/Source/USDExporter/Private/AnimSequenceExporterUSD.cpp b/Engine/Plugins/Importers/USDImporter/Source/USDExporter/Private/AnimSequenceExporterUSD.cpp
--- a/Engine/Plugins/Importers/USDImporter/Source/USDExporter/Private/AnimSequenceExporterUSD.cpp
+++ b/Engine/Plugins/Importers/USDImporter/Source/USDExporter/Private/AnimSequenceExporterUSD.cpp
@@ -4,6 +4,7 @@
 
 #include "AnimSequenceExporterUSDOptions.h"
 #include "EngineAnalytics.h"
+#include "Engine/SkeletalMesh.h"
 #include "MaterialExporterUSD.h"
 #include "UnrealUSDWrapper.h"
 #include "USDClassesModule.h"
diff --git a/Engine/Plugins/Runtime/GameplayAbilities/Source/GameplayAbilities/Private/Abilities/GameplayAbility.cpp b/Engine/Plugins/Runtime/GameplayAbilities/Source/GameplayAbilities/Private/Abilities/GameplayAbility.cpp
--- a/Engine/Plugins/Runtime/GameplayAbilities/Source/GameplayAbilities/Private/Abilities/GameplayAbility.cpp
+++ b/Engine/Plugins/Runtime/GameplayAbilities/Source/GameplayAbilities/Private/Abilities/GameplayAbility.cpp
@@ -9,6 +9,7 @@
 #include "AbilitySystemGlobals.h"
 #include "AbilitySystemComponent.h"
 #include "Abilities/Tasks/AbilityTask.h"
+#include "Components/SkeletalMeshComponent.h"
 #include "GameplayCue_Types.h"
 #include "GameFramework/CharacterMovementComponent.h"
 #include "GameFramework/Character.h"
diff --git a/Engine/Plugins/Runtime/GameplayAbilities/Source/GameplayAbilitiesEditor/Private/SGameplayAttributeGraphPin.cpp b/Engine/Plugins/Runtime/GameplayAbilities/Source/GameplayAbilitiesEditor/Private/SGameplayAttributeGraphPin.cpp
--- a/Engine/Plugins/Runtime/GameplayAbilities/Source/GameplayAbilitiesEditor/Private/SGameplayAttributeGraphPin.cpp
+++ b/Engine/Plugins/Runtime/GameplayAbilities/Source/GameplayAbilitiesEditor/Private/SGameplayAttributeGraphPin.cpp
@@ -6,6 +6,7 @@
 #include "SGameplayAttributeWidget.h"
 #include "ScopedTransaction.h"
 #include "UObject/CoreRedirects.h"
+#include "UObject/UObjectIterator.h"
 #include "AbilitySystemComponent.h"
 
 #define LOCTEXT_NAMESPACE "K2Node"
diff --git a/Engine/Source/Runtime/AnimGraphRuntime/Private/SequenceEvaluatorLibrary.cpp b/Engine/Source/Runtime/AnimGraphRuntime/Private/SequenceEvaluatorLibrary.cpp
--- a/Engine/Source/Runtime/AnimGraphRuntime/Private/SequenceEvaluatorLibrary.cpp
+++ b/Engine/Source/Runtime/AnimGraphRuntime/Private/SequenceEvaluatorLibrary.cpp
@@ -4,6 +4,7 @@
 
 #include "Animation/AnimNode_Inertialization.h"
 #include "AnimNodes/AnimNode_SequenceEvaluator.h"
+#include "AnimationRuntime.h"
 
 DEFINE_LOG_CATEGORY_STATIC(LogSequenceEvaluatorLibrary, Verbose, All);
 
