Fix non unity build errorsFix non unity build errors
#test compile cpp

diff --git a/Engine/Plugins/Experimental/Animation/ContextualAnimation/Source/ContextualAnimation/Public/ContextualAnimTypes.h b/Engine/Plugins/Experimental/Animation/ContextualAnimation/Source/ContextualAnimation/Public/ContextualAnimTypes.h
--- a/Engine/Plugins/Experimental/Animation/ContextualAnimation/Source/ContextualAnimation/Public/ContextualAnimTypes.h
+++ b/Engine/Plugins/Experimental/Animation/ContextualAnimation/Source/ContextualAnimation/Public/ContextualAnimTypes.h
@@ -4,6 +4,7 @@
 
 #include "Animation/AnimSequence.h"
 #include "Templates/SubclassOf.h"
+#include "GameFramework/Actor.h"
 #include "GameplayTagContainer.h"
 #include "ContextualAnimTypes.generated.h"
 
