Preparation pass for TObjectPtr upgrade pass on plugins used when building UnrealEditor with "-allmodules".  Addresses header issues in a way that accommodates either TObjectPtr or raw pointer use.

diff --git a/Engine/Plugins/AI/MLAdapter/Source/MLAdapter/Public/Agents/MLAdapterAgent.h b/Engine/Plugins/AI/MLAdapter/Source/MLAdapter/Public/Agents/MLAdapterAgent.h
--- a/Engine/Plugins/AI/MLAdapter/Source/MLAdapter/Public/Agents/MLAdapterAgent.h
+++ b/Engine/Plugins/AI/MLAdapter/Source/MLAdapter/Public/Agents/MLAdapterAgent.h
@@ -114,8 +114,8 @@ public:
 	/** Get the Controller this agent is controlling. */
 	const AController* GetController() const { return Controller; }
 	
-	TArray<UMLAdapterSensor*>::TConstIterator GetSensorsConstIterator() const { return Sensors.CreateConstIterator(); }
-	TArray<UMLAdapterActuator*>::TConstIterator GetActuatorsConstIterator() const { return Actuators.CreateConstIterator(); }
+	UE_TRANSITIONAL_OBJECT_PTR_TEMPLATE_SUFFIXED(TArray, UMLAdapterSensor, ::TConstIterator) GetSensorsConstIterator() const { return Sensors.CreateConstIterator(); }
+	UE_TRANSITIONAL_OBJECT_PTR_TEMPLATE_SUFFIXED(TArray, UMLAdapterActuator, ::TConstIterator) GetActuatorsConstIterator() const { return Actuators.CreateConstIterator(); }
 	
 	/** If the avatar is a controller, then get the current score from the controller's player state. */
 	virtual float GetReward() const;
diff --git a/Engine/Plugins/Experimental/Animation/ContextualAnimation/Source/ContextualAnimation/Public/ContextualAnimManager.h b/Engine/Plugins/Experimental/Animation/ContextualAnimation/Source/ContextualAnimation/Public/ContextualAnimManager.h
--- a/Engine/Plugins/Experimental/Animation/ContextualAnimation/Source/ContextualAnimation/Public/ContextualAnimManager.h
+++ b/Engine/Plugins/Experimental/Animation/ContextualAnimation/Source/ContextualAnimation/Public/ContextualAnimManager.h
@@ -59,7 +59,7 @@ public:
 	UFUNCTION(BlueprintPure, Category = "Contextual Anim|Manager", meta = (WorldContext = "WorldContextObject"))
 	static UContextualAnimManager* GetContextualAnimManager(UObject* WorldContextObject);
 
-	FORCEINLINE const TSet<UContextualAnimSceneActorComponent*>& GetSceneActorCompContainer() const { return SceneActorCompContainer; };
+	FORCEINLINE const UE_TRANSITIONAL_OBJECT_PTR_TEMPLATE(TSet, UContextualAnimSceneActorComponent)& GetSceneActorCompContainer() const { return SceneActorCompContainer; };
 
 	/** Attempts to start an scene instance with the supplied actors using the first valid set based on selection criteria */
 	UFUNCTION(BlueprintCallable, Category = "Contextual Anim|Manager", meta = (DisplayName = "Try Start Scene"))
diff --git a/Engine/Plugins/Experimental/Enterprise/DatasmithRuntime/Source/Private/DatasmithRuntimeAuxiliaryData.h b/Engine/Plugins/Experimental/Enterprise/DatasmithRuntime/Source/Private/DatasmithRuntimeAuxiliaryData.h
--- a/Engine/Plugins/Experimental/Enterprise/DatasmithRuntime/Source/Private/DatasmithRuntimeAuxiliaryData.h
+++ b/Engine/Plugins/Experimental/Enterprise/DatasmithRuntime/Source/Private/DatasmithRuntimeAuxiliaryData.h
@@ -67,7 +67,7 @@ public:
 
 	virtual const TArray<UAssetUserData*>* GetAssetUserDataArray() const override
 	{
-		return &AssetUserData;
+		return &ToRawPtrTArrayUnsafe(AssetUserData);
 	}
 	//~ End IInterface_AssetUserData Interface
 
diff --git a/Engine/Plugins/Experimental/Water/Source/Runtime/Public/WaterMeshComponent.h b/Engine/Plugins/Experimental/Water/Source/Runtime/Public/WaterMeshComponent.h
--- a/Engine/Plugins/Experimental/Water/Source/Runtime/Public/WaterMeshComponent.h
+++ b/Engine/Plugins/Experimental/Water/Source/Runtime/Public/WaterMeshComponent.h
@@ -48,7 +48,7 @@ public:
 
 	const FWaterQuadTree& GetWaterQuadTree() const { return WaterQuadTree; }
 
-	const TSet<UMaterialInterface*>& GetUsedMaterialsSet() const { return UsedMaterials; }
+	const UE_TRANSITIONAL_OBJECT_PTR_TEMPLATE(TSet, UMaterialInterface)& GetUsedMaterialsSet() const { return UsedMaterials; }
 
 	void MarkWaterMeshGridDirty() { bNeedsRebuild = true; }
 
diff --git a/Engine/Plugins/Importers/USDImporter/Source/USDStage/Public/USDPrimTwin.h b/Engine/Plugins/Importers/USDImporter/Source/USDStage/Public/USDPrimTwin.h
--- a/Engine/Plugins/Importers/USDImporter/Source/USDStage/Public/USDPrimTwin.h
+++ b/Engine/Plugins/Importers/USDImporter/Source/USDStage/Public/USDPrimTwin.h
@@ -24,7 +24,7 @@ class UUsdPrimTwin final : public UObject
 public:
 	UUsdPrimTwin& AddChild( const FString& InPrimPath );
 	void RemoveChild( const TCHAR* InPrimPath );
-	const TMap< FString, UUsdPrimTwin* >& GetChildren() const { return Children; }
+	const UE_TRANSITIONAL_OBJECT_PTR_TEMPLATE2_ARG2(TMap, FString, UUsdPrimTwin)& GetChildren() const { return Children; }
 
 	UUsdPrimTwin* GetParent() const { return Parent.Get(); }
 
@@ -32,7 +32,7 @@ public:
 
 	void Iterate( TFunction< void( UUsdPrimTwin& ) > Func, bool bRecursive )
 	{
-		for ( TMap< FString, UUsdPrimTwin* >::TIterator It = Children.CreateIterator(); It; ++It )
+		for ( UE_TRANSITIONAL_OBJECT_PTR_TEMPLATE2_ARG2_SUFFIXED(TMap, FString, UUsdPrimTwin, ::TIterator) It = Children.CreateIterator(); It; ++It )
 		{
 			UUsdPrimTwin* Child = It->Value;
 
diff --git a/Engine/Plugins/Runtime/MassEntity/Source/MassEntity/Public/MassObserverManager.h b/Engine/Plugins/Runtime/MassEntity/Source/MassEntity/Public/MassObserverManager.h
--- a/Engine/Plugins/Runtime/MassEntity/Source/MassEntity/Public/MassObserverManager.h
+++ b/Engine/Plugins/Runtime/MassEntity/Source/MassEntity/Public/MassObserverManager.h
@@ -20,7 +20,7 @@ struct FMassObserversMap
 	GENERATED_BODY()
 
 	// a helper accessor simplifying access while still keeping Container private
-	TMap<const UScriptStruct*, FMassRuntimePipeline>& operator*()
+	UE_TRANSITIONAL_OBJECT_PTR_TEMPLATE2_ARG1(TMap, const UScriptStruct, FMassRuntimePipeline)& operator*()
 	{
 		return Container;
 	}
diff --git a/Engine/Plugins/Runtime/MassEntity/Source/MassEntity/Public/MassObserverRegistry.h b/Engine/Plugins/Runtime/MassEntity/Source/MassEntity/Public/MassObserverRegistry.h
--- a/Engine/Plugins/Runtime/MassEntity/Source/MassEntity/Public/MassObserverRegistry.h
+++ b/Engine/Plugins/Runtime/MassEntity/Source/MassEntity/Public/MassObserverRegistry.h
@@ -30,12 +30,12 @@ struct FMassEntityObserverClassesMap
 	GENERATED_BODY()
 
 	/** a helper accessor simplifying access while still keeping Container private */
-	const TMap<const UScriptStruct*, FMassProcessorClassCollection>& operator*() const
+	const UE_TRANSITIONAL_OBJECT_PTR_TEMPLATE2_ARG1(TMap, const UScriptStruct, FMassProcessorClassCollection)& operator*() const
 	{
 		return Container;
 	}
 
-	TMap<const UScriptStruct*, FMassProcessorClassCollection>& operator*()
+	UE_TRANSITIONAL_OBJECT_PTR_TEMPLATE2_ARG1(TMap, const UScriptStruct, FMassProcessorClassCollection)& operator*()
 	{
 		return Container;
 	}
diff --git a/Engine/Plugins/VirtualProduction/DMX/DMXEngine/Source/DMXRuntime/Public/Library/DMXEntityReference.h b/Engine/Plugins/VirtualProduction/DMX/DMXEngine/Source/DMXRuntime/Public/Library/DMXEntityReference.h
--- a/Engine/Plugins/VirtualProduction/DMX/DMXEngine/Source/DMXRuntime/Public/Library/DMXEntityReference.h
+++ b/Engine/Plugins/VirtualProduction/DMX/DMXEngine/Source/DMXRuntime/Public/Library/DMXEntityReference.h
@@ -65,7 +65,7 @@ public:
 		// Take the Library address and the EntityID to make a hash from both
 		uint64 ToHash[] =
 		{
-			reinterpret_cast<uintptr_t>(EntityRef.DMXLibrary),
+			reinterpret_cast<uintptr_t>(ToRawPtr(EntityRef.DMXLibrary)),
 			static_cast<uint64>(EntityId64[0]),
 			static_cast<uint64>(EntityId64[1])
 		};
diff --git a/Engine/Source/Runtime/CoreUObject/Public/UObject/ObjectPtr.h b/Engine/Source/Runtime/CoreUObject/Public/UObject/ObjectPtr.h
--- a/Engine/Source/Runtime/CoreUObject/Public/UObject/ObjectPtr.h
+++ b/Engine/Source/Runtime/CoreUObject/Public/UObject/ObjectPtr.h
@@ -24,6 +24,14 @@
  * TObjectPtr<T> after the transition is complete and the type won't be toggling back and forth anymore.
  */
 #define UE_TRANSITIONAL_OBJECT_PTR(Type) auto 
+#define UE_TRANSITIONAL_OBJECT_PTR_TEMPLATE(Templ, Type) auto 
+#define UE_TRANSITIONAL_OBJECT_PTR_TEMPLATE_SUFFIXED(Templ, Type, Suffix) auto 
+#define UE_TRANSITIONAL_OBJECT_PTR_TEMPLATE2_ARG1(Templ, Type1, Type2) auto 
+#define UE_TRANSITIONAL_OBJECT_PTR_TEMPLATE2_ARG1_SUFFIXED(Templ, Type1, Type2, Suffix) auto 
+#define UE_TRANSITIONAL_OBJECT_PTR_TEMPLATE2_ARG2(Templ, Type1, Type2) auto 
+#define UE_TRANSITIONAL_OBJECT_PTR_TEMPLATE2_ARG2_SUFFIXED(Templ, Type1, Type2, Suffix) auto 
+#define UE_TRANSITIONAL_OBJECT_PTR_TEMPLATE2_ARG_BOTH(Templ, Type1, Type2) auto 
+#define UE_TRANSITIONAL_OBJECT_PTR_TEMPLATE2_ARG_BOTH_SUFFIXED(Templ, Type1, Type2, Suffix) auto 
 
 #if PLATFORM_MICROSOFT && defined(_MSC_EXTENSIONS)
 	/**
