Contextual Anim Editor: Refactored new anim set popup into a dialog window to prevent it from disappearing when the user clicks outside the dialog area

diff --git a/Engine/Plugins/Experimental/Animation/ContextualAnimation/Source/ContextualAnimationEditor/Private/ContextualAnimAssetEditorCommands.cpp b/Engine/Plugins/Experimental/Animation/ContextualAnimation/Source/ContextualAnimationEditor/Private/ContextualAnimAssetEditorCommands.cpp
--- a/Engine/Plugins/Experimental/Animation/ContextualAnimation/Source/ContextualAnimationEditor/Private/ContextualAnimAssetEditorCommands.cpp
+++ b/Engine/Plugins/Experimental/Animation/ContextualAnimation/Source/ContextualAnimationEditor/Private/ContextualAnimAssetEditorCommands.cpp
@@ -7,6 +7,7 @@
 void FContextualAnimAssetEditorCommands::RegisterCommands()
 {
 	UI_COMMAND(ResetPreviewScene, "Reset Scene", "Reset Scene.", EUserInterfaceActionType::Button, FInputChord());
+	UI_COMMAND(NewAnimSet, "New AnimSet", "New AnimSet.", EUserInterfaceActionType::Button, FInputChord());
 
 	UI_COMMAND(ShowIKTargetsDrawSelected, "Selected Actor Only", "Show IK Targets for the selected actor", EUserInterfaceActionType::RadioButton, FInputChord());
 	UI_COMMAND(ShowIKTargetsDrawAll, "All Actors", "Show IK Targets for all the actors", EUserInterfaceActionType::RadioButton, FInputChord());
diff --git a/Engine/Plugins/Experimental/Animation/ContextualAnimation/Source/ContextualAnimationEditor/Private/ContextualAnimAssetEditorToolkit.cpp b/Engine/Plugins/Experimental/Animation/ContextualAnimation/Source/ContextualAnimationEditor/Private/ContextualAnimAssetEditorToolkit.cpp
--- a/Engine/Plugins/Experimental/Animation/ContextualAnimation/Source/ContextualAnimationEditor/Private/ContextualAnimAssetEditorToolkit.cpp
+++ b/Engine/Plugins/Experimental/Animation/ContextualAnimation/Source/ContextualAnimationEditor/Private/ContextualAnimAssetEditorToolkit.cpp
@@ -3,6 +3,7 @@
 #include "ContextualAnimAssetEditorToolkit.h"
 #include "SContextualAnimViewport.h"
 #include "SContextualAnimAssetBrowser.h"
+#include "SContextualAnimNewAnimSetDialog.h"
 #include "ContextualAnimPreviewScene.h"
 #include "ContextualAnimAssetEditorCommands.h"
 #include "Widgets/Docking/SDockTab.h"
@@ -157,6 +158,11 @@ void FContextualAnimAssetEditorToolkit::BindCommands()
 		FExecuteAction::CreateSP(this, &FContextualAnimAssetEditorToolkit::ResetPreviewScene),
 		EUIActionRepeatMode::RepeatDisabled);
 
+	ToolkitCommands->MapAction(
+		Commands.NewAnimSet,
+		FExecuteAction::CreateSP(this, &FContextualAnimAssetEditorToolkit::ShowNewAnimSetDialog),
+		EUIActionRepeatMode::RepeatDisabled);
+
 	ToolkitCommands->MapAction(
 		Commands.Simulate,
 		FExecuteAction::CreateSP(this, &FContextualAnimAssetEditorToolkit::ToggleSimulateMode),
@@ -209,11 +215,11 @@ void FContextualAnimAssetEditorToolkit::FillToolbar(FToolBarBuilder& ToolbarBuil
 		FSlateIcon()
 	);
 
-	ToolbarBuilder.AddComboButton(
-		FUIAction(),
-		FOnGetContent::CreateSP(this, &FContextualAnimAssetEditorToolkit::BuildNewAnimSetWidget),
-		LOCTEXT("NewSet_Label", "New Set"),
-		FText::GetEmpty(),
+	ToolbarBuilder.AddToolBarButton(
+		FContextualAnimAssetEditorCommands::Get().NewAnimSet,
+		NAME_None,
+		TAttribute<FText>(),
+		TAttribute<FText>(),
 		FSlateIcon(FAppStyle::GetAppStyleSetName(), "Icons.Plus")
 	);
 
@@ -252,68 +258,10 @@ TSharedRef<SWidget> FContextualAnimAssetEditorToolkit::BuildSectionsMenu()
 	return MenuBuilder.MakeWidget();
 }
 
-TSharedRef<SWidget> FContextualAnimAssetEditorToolkit::BuildNewAnimSetWidget()
+void FContextualAnimAssetEditorToolkit::ShowNewAnimSetDialog()
 {
-	const bool bShouldCloseWindowAfterMenuSelection = true;
-	FMenuBuilder MenuBuilder(bShouldCloseWindowAfterMenuSelection, GetToolkitCommands());
-
-	FPropertyEditorModule& PropertyModule = FModuleManager::LoadModuleChecked<FPropertyEditorModule>("PropertyEditor");
-
-	FDetailsViewArgs Args;
-	Args.bHideSelectionTip = true;
-	Args.bAllowSearch = false;
-	Args.bAllowFavoriteSystem = false;
-
-	NewAnimSetWidgetStruct = MakeShared<FStructOnScope>(FContextualAnimNewAnimSetParams::StaticStruct());
-	FContextualAnimNewAnimSetParams* Params = (FContextualAnimNewAnimSetParams*)NewAnimSetWidgetStruct->GetStructMemory();
-
-	const TArray<FName> Roles = ViewModel->GetSceneAsset()->GetRoles();
-	for (FName Role : Roles)
-	{
-		FContextualAnimNewAnimSetData Entry;
-		Entry.RoleName = Role;
-		Params->Data.Add(Entry);
-	}
-
-	TSharedRef<IStructureDetailsView> StructureDetailsView = PropertyModule.CreateStructureDetailView(Args, FStructureDetailsViewArgs(), NewAnimSetWidgetStruct);
-
-	MenuBuilder.AddWidget(
-		SNew(SBox)
-		.MinDesiredWidth(500.0f)
-		.MaxDesiredWidth(500.f)
-		.MaxDesiredHeight(400.0f)
-		[
-			SNew(SVerticalBox)
-			+ SVerticalBox::Slot()
-			.AutoHeight()
-			.HAlign(HAlign_Fill)
-			[
-				StructureDetailsView->GetWidget().ToSharedRef()
-			]
-			+ SVerticalBox::Slot()
-			.AutoHeight()
-			.HAlign(HAlign_Fill)
-			.Padding(5)
-			[
-				SNew(SButton)
-				.HAlign(HAlign_Center)
-				.ContentPadding(FAppStyle::GetMargin("StandardDialog.ContentPadding"))
-				.Text(LOCTEXT("OK", "OK"))
-				.OnClicked_Lambda([this]()
-				{
-					FContextualAnimNewAnimSetParams* Params = (FContextualAnimNewAnimSetParams*)NewAnimSetWidgetStruct->GetStructMemory();
-					check(Params);
-
-					ViewModel->AddNewAnimSet(*Params);
-
-					FSlateApplication::Get().DismissAllMenus();
-
-					return FReply::Handled();
-				})
-			]
-		], FText(), true, false);
-
-	return MenuBuilder.MakeWidget();
+	TSharedRef<SContextualAnimNewAnimSetDialog> NewAnimSetDialog = SNew(SContextualAnimNewAnimSetDialog, ViewModel.ToSharedRef());
+	NewAnimSetDialog->Show();
 }
 
 void FContextualAnimAssetEditorToolkit::RegisterTabSpawners(const TSharedRef<FTabManager>& InTabManager)
diff --git a/Engine/Plugins/Experimental/Animation/ContextualAnimation/Source/ContextualAnimationEditor/Private/SContextualAnimNewAnimSetDialog.cpp b/Engine/Plugins/Experimental/Animation/ContextualAnimation/Source/ContextualAnimationEditor/Private/SContextualAnimNewAnimSetDialog.cpp
new file mode 100644
--- /dev/null
+++ b/Engine/Plugins/Experimental/Animation/ContextualAnimation/Source/ContextualAnimationEditor/Private/SContextualAnimNewAnimSetDialog.cpp
@@ -0,0 +1,107 @@
+// Copyright Epic Games, Inc. All Rights Reserved.
+
+#include "SContextualAnimNewAnimSetDialog.h"
+#include "Modules/ModuleManager.h"
+#include "PropertyEditorModule.h"
+#include "IStructureDetailsView.h"
+#include "UObject/StructOnScope.h"
+#include "Widgets/Layout/SUniformGridPanel.h"
+#include "Widgets/Input/SButton.h"
+#include "ContextualAnimEditorTypes.h"
+#include "ContextualAnimViewModel.h"
+#include "ContextualAnimSceneAsset.h"
+#include "Editor.h"
+
+#define LOCTEXT_NAMESPACE "ContextualAnimNewAnimSetDialog"
+
+void SContextualAnimNewAnimSetDialog::Construct(const FArguments& InArgs, const TSharedRef<FContextualAnimViewModel>& InViewModel)
+{
+	WeakViewModel = InViewModel;
+
+	WidgetStruct = MakeShared<FStructOnScope>(FContextualAnimNewAnimSetParams::StaticStruct());
+	FContextualAnimNewAnimSetParams* Params = (FContextualAnimNewAnimSetParams*)WidgetStruct->GetStructMemory();
+
+	const TArray<FName> Roles = InViewModel->GetSceneAsset()->GetRoles();
+	for (FName Role : Roles)
+	{
+		FContextualAnimNewAnimSetData Entry;
+		Entry.RoleName = Role;
+		Params->Data.Add(Entry);
+	}
+
+	FDetailsViewArgs Args;
+	Args.bHideSelectionTip = true;
+	Args.bAllowSearch = false;
+	Args.bAllowFavoriteSystem = false;
+
+	FPropertyEditorModule& PropertyModule = FModuleManager::LoadModuleChecked<FPropertyEditorModule>("PropertyEditor");
+	TSharedRef<IStructureDetailsView> StructureDetailsView = PropertyModule.CreateStructureDetailView(Args, FStructureDetailsViewArgs(), WidgetStruct);
+
+	SWindow::Construct(SWindow::FArguments()
+		.Title(LOCTEXT("NewAnimSetTitle", "Add New AnimSet"))
+		.SupportsMinimize(false)
+		.SupportsMaximize(false)
+		.ClientSize(FVector2D(500, 400))
+		[
+			SNew(SVerticalBox)
+			+ SVerticalBox::Slot()
+			.FillHeight(1)
+			.Padding(3)
+			[
+				StructureDetailsView->GetWidget().ToSharedRef()
+			]
+
+			+ SVerticalBox::Slot()
+			.AutoHeight()
+			.HAlign(HAlign_Right)
+			.Padding(5)
+			[
+				SNew(SUniformGridPanel)
+				.SlotPadding(FAppStyle::GetMargin("StandardDialog.SlotPadding"))
+				.MinDesiredSlotWidth(FAppStyle::GetFloat("StandardDialog.MinDesiredSlotWidth"))
+				.MinDesiredSlotHeight(FAppStyle::GetFloat("StandardDialog.MinDesiredSlotHeight"))
+				+ SUniformGridPanel::Slot(0, 0)
+				[
+					SNew(SButton)
+					.HAlign(HAlign_Center)
+					.ContentPadding(FAppStyle::GetMargin("StandardDialog.ContentPadding"))
+					.Text(LOCTEXT("OK", "OK"))
+					.OnClicked(this, &SContextualAnimNewAnimSetDialog::OnButtonClick, EAppReturnType::Ok)
+				]
+				+ SUniformGridPanel::Slot(1, 0)
+				[
+					SNew(SButton)
+					.HAlign(HAlign_Center)
+					.ContentPadding(FAppStyle::GetMargin("StandardDialog.ContentPadding"))
+					.Text(LOCTEXT("Cancel", "Cancel"))
+					.OnClicked(this, &SContextualAnimNewAnimSetDialog::OnButtonClick, EAppReturnType::Cancel)
+				]
+			]
+		]);
+}
+
+
+FReply SContextualAnimNewAnimSetDialog::OnButtonClick(EAppReturnType::Type ButtonID)
+{
+	UserResponse = ButtonID;
+
+	if(ButtonID == EAppReturnType::Ok)
+	{
+		FContextualAnimNewAnimSetParams* Params = (FContextualAnimNewAnimSetParams*) WidgetStruct->GetStructMemory();
+		check(Params);
+
+		WeakViewModel.Pin()->AddNewAnimSet(*Params);
+	}
+
+	RequestDestroyWindow();
+
+	return FReply::Handled();
+}
+
+EAppReturnType::Type SContextualAnimNewAnimSetDialog::Show()
+{
+	GEditor->EditorAddModalWindow(SharedThis(this));
+	return UserResponse;
+}
+
+#undef LOCTEXT_NAMESPACE
\ No newline at end of file
diff --git a/Engine/Plugins/Experimental/Animation/ContextualAnimation/Source/ContextualAnimationEditor/Public/ContextualAnimAssetEditorCommands.h b/Engine/Plugins/Experimental/Animation/ContextualAnimation/Source/ContextualAnimationEditor/Public/ContextualAnimAssetEditorCommands.h
--- a/Engine/Plugins/Experimental/Animation/ContextualAnimation/Source/ContextualAnimationEditor/Public/ContextualAnimAssetEditorCommands.h
+++ b/Engine/Plugins/Experimental/Animation/ContextualAnimation/Source/ContextualAnimationEditor/Public/ContextualAnimAssetEditorCommands.h
@@ -20,6 +20,8 @@ public:
 
 	TSharedPtr<FUICommandInfo> ResetPreviewScene;
 
+	TSharedPtr<FUICommandInfo> NewAnimSet;
+
 	TSharedPtr<FUICommandInfo> ShowIKTargetsDrawSelected;
 
 	TSharedPtr<FUICommandInfo> ShowIKTargetsDrawAll;
diff --git a/Engine/Plugins/Experimental/Animation/ContextualAnimation/Source/ContextualAnimationEditor/Public/ContextualAnimAssetEditorToolkit.h b/Engine/Plugins/Experimental/Animation/ContextualAnimation/Source/ContextualAnimationEditor/Public/ContextualAnimAssetEditorToolkit.h
--- a/Engine/Plugins/Experimental/Animation/ContextualAnimation/Source/ContextualAnimationEditor/Public/ContextualAnimAssetEditorToolkit.h
+++ b/Engine/Plugins/Experimental/Animation/ContextualAnimation/Source/ContextualAnimationEditor/Public/ContextualAnimAssetEditorToolkit.h
@@ -53,7 +53,8 @@ private:
 	void FillToolbar(FToolBarBuilder& ToolbarBuilder);
 
 	TSharedRef<SWidget> BuildSectionsMenu();
-	TSharedRef<SWidget> BuildNewAnimSetWidget();
+
+	void ShowNewAnimSetDialog();
 
 	void OnFinishedChangingProperties(const FPropertyChangedEvent& PropertyChangedEvent);
 
@@ -66,8 +67,6 @@ private:
 	TSharedPtr<FContextualAnimPreviewScene> PreviewScene;
 
 	TSharedPtr<FContextualAnimViewModel> ViewModel;
-
-	TSharedPtr<FStructOnScope> NewAnimSetWidgetStruct;
 };
 
 
diff --git a/Engine/Plugins/Experimental/Animation/ContextualAnimation/Source/ContextualAnimationEditor/Public/SContextualAnimNewAnimSetDialog.h b/Engine/Plugins/Experimental/Animation/ContextualAnimation/Source/ContextualAnimationEditor/Public/SContextualAnimNewAnimSetDialog.h
new file mode 100644
--- /dev/null
+++ b/Engine/Plugins/Experimental/Animation/ContextualAnimation/Source/ContextualAnimationEditor/Public/SContextualAnimNewAnimSetDialog.h
@@ -0,0 +1,36 @@
+// Copyright Epic Games, Inc. All Rights Reserved.
+
+#pragma once
+
+#include "Widgets/SWindow.h"
+
+class FContextualAnimViewModel;
+class FStructOnScope;
+
+/** Dialog for adding a new AnimSet */
+class SContextualAnimNewAnimSetDialog : public SWindow
+{
+public:
+
+	SLATE_BEGIN_ARGS(SContextualAnimNewAnimSetDialog){}
+	SLATE_END_ARGS()
+
+	SContextualAnimNewAnimSetDialog() 
+		: UserResponse(EAppReturnType::Cancel)
+	{
+	}
+
+	void Construct(const FArguments& InArgs, const TSharedRef<FContextualAnimViewModel>& InViewModel);
+
+	EAppReturnType::Type Show();
+
+protected:
+
+	TWeakPtr<FContextualAnimViewModel> WeakViewModel;
+
+	EAppReturnType::Type UserResponse;
+
+	TSharedPtr<FStructOnScope> WidgetStruct;
+
+	FReply OnButtonClick(EAppReturnType::Type ButtonID);
+};
