Contextual Anim: Added selection criterion based on distance. Extended selection criterion based on angle. Added verbose logs to each criterion for debug purposes

diff --git a/Engine/Plugins/Experimental/Animation/ContextualAnimation/Source/ContextualAnimation/Private/ContextualAnimSelectionCriterion.cpp b/Engine/Plugins/Experimental/Animation/ContextualAnimation/Source/ContextualAnimation/Private/ContextualAnimSelectionCriterion.cpp
--- a/Engine/Plugins/Experimental/Animation/ContextualAnimation/Source/ContextualAnimation/Private/ContextualAnimSelectionCriterion.cpp
+++ b/Engine/Plugins/Experimental/Animation/ContextualAnimation/Source/ContextualAnimation/Private/ContextualAnimSelectionCriterion.cpp
@@ -32,18 +32,17 @@ UContextualAnimSelectionCriterion_TriggerArea::UContextualAnimSelectionCriterion
 
 bool UContextualAnimSelectionCriterion_TriggerArea::DoesQuerierPassCondition(const FContextualAnimSceneBindingContext& Primary, const FContextualAnimSceneBindingContext& Querier) const
 {
-	if (PolygonPoints.Num() == 4)
-	{
-		const FTransform PrimaryTransform = Primary.GetTransform();
-		const FTransform QuerierTransform = Querier.GetTransform();
+	check(PolygonPoints.Num() == 4);
 
-		const float HalfHeight = FMath::Max((Height / 2.f), 0.f);
-		const float VDist = FMath::Abs((PrimaryTransform.GetLocation().Z + PolygonPoints[0].Z + HalfHeight) - QuerierTransform.GetLocation().Z);
-		if (VDist > HalfHeight)
-		{
-			return false;
-		}
+	bool bResult = false;
 
+	const FTransform PrimaryTransform = Primary.GetTransform();
+	const FTransform QuerierTransform = Querier.GetTransform();
+
+	const float HalfHeight = FMath::Max((Height / 2.f), 0.f);
+	const float VDist = FMath::Abs((PrimaryTransform.GetLocation().Z + PolygonPoints[0].Z + HalfHeight) - QuerierTransform.GetLocation().Z);
+	if (VDist <= HalfHeight)
+	{
 		const FVector2D TestPoint = FVector2D(QuerierTransform.GetLocation());
 		const int32 NumPoints = PolygonPoints.Num();
 		float AngleSum = 0.0f;
@@ -55,27 +54,88 @@ bool UContextualAnimSelectionCriterion_TriggerArea::DoesQuerierPassCondition(con
 			AngleSum += Angle;
 		}
 
-		return (FMath::Abs(AngleSum) > 0.001f);
+		bResult = (FMath::Abs(AngleSum) > 0.001f);
 	}
 
-	return false;
+	UE_LOG(LogContextualAnim, Verbose, TEXT("UContextualAnimSelectionCriterion_TriggerArea: Primary: %s Querier: %s Result: %d"),
+		*GetNameSafe(Primary.GetActor()), *GetNameSafe(Querier.GetActor()), bResult);
+
+	return bResult;
 }
 
-// UContextualAnimSelectionCriterion_Facing
+// UContextualAnimSelectionCriterion_Angle
 //===========================================================================
 
-bool UContextualAnimSelectionCriterion_Facing::DoesQuerierPassCondition(const FContextualAnimSceneBindingContext& Primary, const FContextualAnimSceneBindingContext& Querier) const
+bool UContextualAnimSelectionCriterion_Angle::DoesQuerierPassCondition(const FContextualAnimSceneBindingContext& Primary, const FContextualAnimSceneBindingContext& Querier) const
 {
-	if (MaxAngle > 0.f)
+	const FTransform PrimaryTransform = Primary.GetTransform();
+	const FTransform QuerierTransform = Querier.GetTransform();
+
+	auto CalculateAngle = [](const FTransform& A, const FTransform& B, bool bSignedAngle)
 	{
-		const FTransform PrimaryTransform = Primary.GetTransform();
-		const FTransform QuerierTransform = Querier.GetTransform();
+		const FVector ToTarget = (A.GetLocation() - B.GetLocation()).GetSafeNormal2D();
+		const float ForwardCosAngle = FVector::DotProduct(B.GetRotation().GetForwardVector(), ToTarget);
+		const float ForwardDeltaDegree = FMath::RadiansToDegrees(FMath::Acos(ForwardCosAngle));
+		
+		if(bSignedAngle)
+		{
+			const float RightCosAngle = FVector::DotProduct(B.GetRotation().GetRightVector(), ToTarget);
+			return (RightCosAngle < 0) ? (ForwardDeltaDegree * -1) : ForwardDeltaDegree;
+		}
+		else
+		{
+			return ForwardDeltaDegree;
+		}
+	};
+
+	float Angle = 0.f;
+	if (Mode == EContextualAnimCriterionAngleMode::ToPrimary)
+	{
+		Angle = CalculateAngle(PrimaryTransform, QuerierTransform, bUseSignedAngle);
+	}
+	else if (Mode == EContextualAnimCriterionAngleMode::FromPrimary)
+	{
+		Angle = CalculateAngle(QuerierTransform, PrimaryTransform, bUseSignedAngle);
+	}
 
-		const float FacingCos = FMath::Cos(FMath::Clamp(FMath::DegreesToRadians(MaxAngle), 0.f, PI));
-		const FVector ToTarget = (PrimaryTransform.GetLocation() - QuerierTransform.GetLocation()).GetSafeNormal2D();
-		const float DotProduct = FVector::DotProduct(QuerierTransform.GetRotation().GetForwardVector(), ToTarget);
-		return (DotProduct > FacingCos);
+	bool bResult = false;
+	if (MinAngle <= 0 && MaxAngle >= 0)
+	{
+		bResult = FMath::IsWithinInclusive(FMath::Abs(Angle), FMath::Abs(MinAngle), FMath::Abs(MaxAngle));
 	}
+	else
+	{
+		bResult = FMath::IsWithinInclusive(Angle, MinAngle, MaxAngle);
+	}
+
+	UE_LOG(LogContextualAnim, Verbose, TEXT("UContextualAnimSelectionCriterion_Angle: Primary: %s Querier: %s Mode: %s bUseSignedAngle: %d MinAngle: %.1f MaxAngle: %.1f Angle: %.1f Result: %d"), 
+		*GetNameSafe(Primary.GetActor()), *GetNameSafe(Querier.GetActor()), *UEnum::GetValueAsString(TEXT("ContextualAnimation.EContextualAnimCriterionAngleMode"), Mode), bUseSignedAngle, MinAngle, MaxAngle, Angle, bResult);
+
+	return bResult;
+}
+
+// UContextualAnimSelectionCriterion_Distance
+//===========================================================================
+
+bool UContextualAnimSelectionCriterion_Distance::DoesQuerierPassCondition(const FContextualAnimSceneBindingContext& Primary, const FContextualAnimSceneBindingContext& Querier) const
+{
+	const FTransform PrimaryTransform = Primary.GetTransform();
+	const FTransform QuerierTransform = Querier.GetTransform();
+
+	float Distance = 0.f;
+	if (Mode == EContextualAnimCriterionDistanceMode::Distance_2D)
+	{
+		Distance = FVector::Dist2D(PrimaryTransform.GetLocation(), QuerierTransform.GetLocation());
+	}
+	else if (Mode == EContextualAnimCriterionDistanceMode::Distance_3D)
+	{
+		Distance = FVector::Dist(PrimaryTransform.GetLocation(), QuerierTransform.GetLocation());
+	}
+
+	const bool bResult = FMath::IsWithinInclusive(Distance, MinDistance, MaxDistance);
+
+	UE_LOG(LogContextualAnim, Verbose, TEXT("UContextualAnimSelectionCriterion_Distance: Primary: %s Querier: %s Mode: %s MaxDistance: %.1f MaxDist: %.1f Dist: %.1f Result: %d"),
+		*GetNameSafe(Primary.GetActor()), *GetNameSafe(Querier.GetActor()), *UEnum::GetValueAsString(TEXT("ContextualAnimation.EContextualAnimCriterionDistanceMode"), Mode), MinDistance, MaxDistance, Distance, bResult);
 
-	return true;
+	return bResult;
 }
\ No newline at end of file
diff --git a/Engine/Plugins/Experimental/Animation/ContextualAnimation/Source/ContextualAnimation/Public/ContextualAnimSelectionCriterion.h b/Engine/Plugins/Experimental/Animation/ContextualAnimation/Source/ContextualAnimation/Public/ContextualAnimSelectionCriterion.h
--- a/Engine/Plugins/Experimental/Animation/ContextualAnimation/Source/ContextualAnimation/Public/ContextualAnimSelectionCriterion.h
+++ b/Engine/Plugins/Experimental/Animation/ContextualAnimation/Source/ContextualAnimation/Public/ContextualAnimSelectionCriterion.h
@@ -55,20 +55,70 @@ public:
 	virtual bool DoesQuerierPassCondition(const FContextualAnimSceneBindingContext& Primary, const FContextualAnimSceneBindingContext& Querier) const override;
 };
 
-// UContextualAnimSelectionCriterion_Facing
+// UContextualAnimSelectionCriterion_Angle
 //===========================================================================
 
+UENUM(BlueprintType)
+enum class EContextualAnimCriterionAngleMode : uint8
+{
+	/** Uses the angle between the vector from querier to primary and querier forward vector */
+	ToPrimary,
+
+	/** Uses the angle between the vector from primary to querier and primary forward vector */
+	FromPrimary
+};
+
 UCLASS()
-class CONTEXTUALANIMATION_API UContextualAnimSelectionCriterion_Facing : public UContextualAnimSelectionCriterion
+class CONTEXTUALANIMATION_API UContextualAnimSelectionCriterion_Angle : public UContextualAnimSelectionCriterion
 {
 	GENERATED_BODY()
 
 public:
 
-	UPROPERTY(EditAnywhere, BlueprintReadWrite, Category = "Defaults", meta = (ClampMin = "0", ClampMax = "180", UIMin = "0", UIMax = "180"))
+	UPROPERTY(EditAnywhere, BlueprintReadWrite, Category = "Defaults")
+	EContextualAnimCriterionAngleMode Mode = EContextualAnimCriterionAngleMode::ToPrimary;
+
+	UPROPERTY(EditAnywhere, BlueprintReadWrite, Category = "Defaults")
+	bool bUseSignedAngle = false;
+
+	UPROPERTY(EditAnywhere, BlueprintReadWrite, Category = "Defaults", meta = (ClampMin = "-180", UIMin = "-180", ClampMax = "180", UIMax = "180"))
+	float MinAngle = 0.f;
+
+	UPROPERTY(EditAnywhere, BlueprintReadWrite, Category = "Defaults", meta = (ClampMin = "-180", UIMin = "-180", ClampMax = "180", UIMax = "180"))
 	float MaxAngle = 0.f;
 
-	UContextualAnimSelectionCriterion_Facing(const FObjectInitializer& ObjectInitializer) : Super(ObjectInitializer) {}
+	UContextualAnimSelectionCriterion_Angle(const FObjectInitializer& ObjectInitializer) : Super(ObjectInitializer) {}
+
+	virtual bool DoesQuerierPassCondition(const FContextualAnimSceneBindingContext& Primary, const FContextualAnimSceneBindingContext& Querier) const override;
+};
+
+// UContextualAnimSelectionCriterion_Distance
+//===========================================================================
+
+UENUM(BlueprintType)
+enum class EContextualAnimCriterionDistanceMode : uint8
+{
+	Distance_3D,
+	Distance_2D
+};
+
+UCLASS()
+class CONTEXTUALANIMATION_API UContextualAnimSelectionCriterion_Distance : public UContextualAnimSelectionCriterion
+{
+	GENERATED_BODY()
+
+public:
+
+	UPROPERTY(EditAnywhere, BlueprintReadWrite, Category = "Defaults")
+	EContextualAnimCriterionDistanceMode Mode = EContextualAnimCriterionDistanceMode::Distance_2D;
+
+	UPROPERTY(EditAnywhere, BlueprintReadWrite, Category = "Defaults", meta = (ClampMin = "0", UIMin = "0"))
+	float MinDistance = 0.f;
+
+	UPROPERTY(EditAnywhere, BlueprintReadWrite, Category = "Defaults", meta = (ClampMin = "0", UIMin = "0"))
+	float MaxDistance = 0.f;
+
+	UContextualAnimSelectionCriterion_Distance(const FObjectInitializer& ObjectInitializer) : Super(ObjectInitializer) {}
 
 	virtual bool DoesQuerierPassCondition(const FContextualAnimSceneBindingContext& Primary, const FContextualAnimSceneBindingContext& Querier) const override;
 };
\ No newline at end of file
diff --git a/Engine/Plugins/Experimental/Animation/ContextualAnimation/Source/ContextualAnimationEditor/Private/ContextualAnimEdMode.cpp b/Engine/Plugins/Experimental/Animation/ContextualAnimation/Source/ContextualAnimationEditor/Private/ContextualAnimEdMode.cpp
--- a/Engine/Plugins/Experimental/Animation/ContextualAnimation/Source/ContextualAnimationEditor/Private/ContextualAnimEdMode.cpp
+++ b/Engine/Plugins/Experimental/Animation/ContextualAnimation/Source/ContextualAnimationEditor/Private/ContextualAnimEdMode.cpp
@@ -114,11 +114,6 @@ void FContextualAnimEdMode::Render(const FSceneView* View, FViewport* Viewport,
 									PDI->SetHitProxy(nullptr);
 								}
 							}
-							else if (const UContextualAnimSelectionCriterion_Facing* Facing = Cast<UContextualAnimSelectionCriterion_Facing>(Criterion))
-							{
-								const FTransform Transform = AnimTrack.GetAlignmentTransformAtEntryTime() * PrimaryTransform;
-								UContextualAnimUtilities::DrawSector(*PDI, Transform.GetLocation(), Transform.GetRotation().GetForwardVector(), 0.f, 30.f, -Facing->MaxAngle, Facing->MaxAngle, DrawColor, SDPG_World, 1.f);
-							}
 						}
 					}
 				}
