Contextual Anim: Exposed multiple functions to blueprint

diff --git a/Engine/Plugins/Experimental/Animation/ContextualAnimation/Source/ContextualAnimation/Private/ContextualAnimSceneAsset.cpp b/Engine/Plugins/Experimental/Animation/ContextualAnimation/Source/ContextualAnimation/Private/ContextualAnimSceneAsset.cpp
--- a/Engine/Plugins/Experimental/Animation/ContextualAnimation/Source/ContextualAnimation/Private/ContextualAnimSceneAsset.cpp
+++ b/Engine/Plugins/Experimental/Animation/ContextualAnimation/Source/ContextualAnimation/Private/ContextualAnimSceneAsset.cpp
@@ -767,3 +767,12 @@ int32 UContextualAnimSceneAsset::BP_FindAnimSetIndexByAnimation(int32 SectionIdx
 
 	return INDEX_NONE;
 }
+
+void UContextualAnimSceneAsset::BP_GetStartAndEndTimeForWarpSection(int32 SectionIdx, int32 AnimSetIdx, FName Role, FName WarpSectionName, float& OutStartTime, float& OutEndTime) const
+{
+	OutStartTime = OutEndTime = 0.f;
+	if(const FContextualAnimTrack* AnimTrack = GetAnimTrack(SectionIdx, AnimSetIdx, Role))
+	{
+		AnimTrack->GetStartAndEndTimeForWarpSection(WarpSectionName, OutStartTime, OutEndTime);
+	}
+}
diff --git a/Engine/Plugins/Experimental/Animation/ContextualAnimation/Source/ContextualAnimation/Private/ContextualAnimTypes.cpp b/Engine/Plugins/Experimental/Animation/ContextualAnimation/Source/ContextualAnimation/Private/ContextualAnimTypes.cpp
--- a/Engine/Plugins/Experimental/Animation/ContextualAnimation/Source/ContextualAnimation/Private/ContextualAnimTypes.cpp
+++ b/Engine/Plugins/Experimental/Animation/ContextualAnimation/Source/ContextualAnimation/Private/ContextualAnimTypes.cpp
@@ -45,13 +45,28 @@ FTransform FContextualAnimAlignmentTrackContainer::ExtractTransformAtTime(int32
 }
 
 float FContextualAnimTrack::GetSyncTimeForWarpSection(int32 WarpSectionIndex) const
+{
+	float StartTime, EndTime;
+	GetStartAndEndTimeForWarpSection(WarpSectionIndex, StartTime, EndTime);
+	return EndTime;
+}
+
+float FContextualAnimTrack::GetSyncTimeForWarpSection(const FName& WarpSectionName) const
+{
+	float StartTime, EndTime;
+	GetStartAndEndTimeForWarpSection(WarpSectionName, StartTime, EndTime);
+	return EndTime;
+}
+
+void FContextualAnimTrack::GetStartAndEndTimeForWarpSection(int32 WarpSectionIndex, float& OutStartTime, float& OutEndTime) const
 {
 	//@TODO: We need a better way to identify warping sections withing the animation. This is just a temp solution
 	//@TODO: We should cache this data
 
-	float Result = 0.f;
+	OutStartTime = 0.f;
+	OutEndTime = 0.f;
 
-	if(Animation && WarpSectionIndex >= 0)
+	if (Animation && WarpSectionIndex >= 0)
 	{
 		FName LastWarpTargetName = NAME_None;
 		int32 LastWarpSectionIndex = INDEX_NONE;
@@ -59,38 +74,41 @@ float FContextualAnimTrack::GetSyncTimeForWarpSection(int32 WarpSectionIndex) co
 		for (int32 Idx = 0; Idx < Animation->Notifies.Num(); Idx++)
 		{
 			const FAnimNotifyEvent& NotifyEvent = Animation->Notifies[Idx];
-			if(const UAnimNotifyState_MotionWarping* Notify = Cast<const UAnimNotifyState_MotionWarping>(NotifyEvent.NotifyStateClass))
+			if (const UAnimNotifyState_MotionWarping* Notify = Cast<const UAnimNotifyState_MotionWarping>(NotifyEvent.NotifyStateClass))
 			{
-				if(const URootMotionModifier_Warp* Modifier = Cast<const URootMotionModifier_Warp>(Notify->RootMotionModifier))
+				if (const URootMotionModifier_Warp* Modifier = Cast<const URootMotionModifier_Warp>(Notify->RootMotionModifier))
 				{
 					const FName WarpTargetName = Modifier->WarpTargetName;
-					if(WarpTargetName != NAME_None)
+					if (WarpTargetName != NAME_None)
 					{
 						// First valid warping window. Initialize everything
 						if (LastWarpSectionIndex == INDEX_NONE)
 						{
 							LastWarpTargetName = WarpTargetName;
-							Result = NotifyEvent.GetEndTriggerTime();
+							OutStartTime = NotifyEvent.GetTriggerTime();
+							OutEndTime = NotifyEvent.GetEndTriggerTime();
 							LastWarpSectionIndex = 0;
 						}
 						// If we hit another warping window but the sync point is the same as the previous, update SyncTime.
 						// This is to deal with cases where a first short window is used to face the alignment point and a second one to perform the rest of the warp
-						else if(WarpTargetName == LastWarpTargetName)
+						else if (WarpTargetName == LastWarpTargetName)
 						{
-							Result = NotifyEvent.GetEndTriggerTime();
+							OutStartTime = NotifyEvent.GetTriggerTime();
+							OutEndTime = NotifyEvent.GetEndTriggerTime();
 						}
 						// If we hit another warping window but with a different sync point name means that we have hit the first window of another warping section
 						else
 						{
 							// If we haven't reached the desired WarpSection yet. Update control vars and keep moving
-							if(WarpSectionIndex > LastWarpSectionIndex)
+							if (WarpSectionIndex > LastWarpSectionIndex)
 							{
 								LastWarpTargetName = WarpTargetName;
-								Result = NotifyEvent.GetEndTriggerTime();
+								OutStartTime = NotifyEvent.GetTriggerTime();
+								OutEndTime = NotifyEvent.GetEndTriggerTime();
 								LastWarpSectionIndex++;
 							}
 							// Otherwise, stop here and return the value of the last window we found
-							else 
+							else
 							{
 								break;
 							}
@@ -100,17 +118,18 @@ float FContextualAnimTrack::GetSyncTimeForWarpSection(int32 WarpSectionIndex) co
 			}
 		}
 	}
-
-	return Result;
 }
 
-float FContextualAnimTrack::GetSyncTimeForWarpSection(const FName& WarpSectionName) const
+void FContextualAnimTrack::GetStartAndEndTimeForWarpSection(const FName& WarpSectionName, float& OutStartTime, float& OutEndTime) const
 {
 	//@TODO: We need a better way to identify warping sections within the animation. This is just a temp solution
 	//@TODO: We should cache this data
 
-	float Result = 0.f;
+	OutStartTime = 0.f;
+	OutEndTime = 0.f;
 
+	int32 Index = INDEX_NONE;
+	float LastEndTime = 0.f;
 	if (Animation && WarpSectionName != NAME_None)
 	{
 		for (int32 Idx = 0; Idx < Animation->Notifies.Num(); Idx++)
@@ -124,9 +143,10 @@ float FContextualAnimTrack::GetSyncTimeForWarpSection(const FName& WarpSectionNa
 					if (WarpSectionName == WarpTargetName)
 					{
 						const float NotifyEndTriggerTime = NotifyEvent.GetEndTriggerTime();
-						if(NotifyEndTriggerTime > Result)
+						if(NotifyEndTriggerTime > LastEndTime)
 						{
-							Result = NotifyEndTriggerTime;
+							LastEndTime = NotifyEndTriggerTime;
+							Index = Idx;
 						}
 					}
 				}
@@ -134,7 +154,12 @@ float FContextualAnimTrack::GetSyncTimeForWarpSection(const FName& WarpSectionNa
 		}
 	}
 
-	return Result;
+	if(Index != INDEX_NONE)
+	{
+		const FAnimNotifyEvent& NotifyEvent = Animation->Notifies[Index];
+		OutStartTime = NotifyEvent.GetTriggerTime();
+		OutEndTime = NotifyEvent.GetEndTriggerTime();
+	}
 }
 
 float FContextualAnimTrack::FindBestAnimStartTime(const FVector& LocalLocation) const
diff --git a/Engine/Plugins/Experimental/Animation/ContextualAnimation/Source/ContextualAnimation/Private/ContextualAnimUtilities.cpp b/Engine/Plugins/Experimental/Animation/ContextualAnimation/Source/ContextualAnimation/Private/ContextualAnimUtilities.cpp
--- a/Engine/Plugins/Experimental/Animation/ContextualAnimation/Source/ContextualAnimation/Private/ContextualAnimUtilities.cpp
+++ b/Engine/Plugins/Experimental/Animation/ContextualAnimation/Source/ContextualAnimation/Private/ContextualAnimUtilities.cpp
@@ -282,6 +282,12 @@ bool UContextualAnimUtilities::BP_CreateContextualAnimSceneBindings(const UConte
 // SceneBindings Blueprint Interface
 //------------------------------------------------------------------------------------------
 
+void UContextualAnimUtilities::BP_SceneBindings_GetSectionAndAnimSetIndices(const FContextualAnimSceneBindings& Bindings, int32& SectionIdx, int32& AnimSetIdx)
+{
+	SectionIdx = Bindings.GetSectionIdx();
+	AnimSetIdx = Bindings.GetAnimSetIdx();
+}
+
 const FContextualAnimSceneBinding& UContextualAnimUtilities::BP_SceneBindings_GetBindingByRole(const FContextualAnimSceneBindings& Bindings, FName Role)
 {
 	if(const FContextualAnimSceneBinding* SceneActorData = Bindings.FindBindingByRole(Role))
@@ -292,6 +298,16 @@ const FContextualAnimSceneBinding& UContextualAnimUtilities::BP_SceneBindings_Ge
 	return FContextualAnimSceneBinding::InvalidBinding;
 }
 
+const FContextualAnimSceneBinding& UContextualAnimUtilities::BP_SceneBindings_GetBindingByActor(const FContextualAnimSceneBindings& Bindings, const AActor* Actor)
+{
+	if (const FContextualAnimSceneBinding* SceneActorData = Bindings.FindBindingByActor(Actor))
+	{
+		return *SceneActorData;
+	}
+
+	return FContextualAnimSceneBinding::InvalidBinding;
+}
+
 void UContextualAnimUtilities::BP_SceneBindings_CalculateAnimSetPivots(const FContextualAnimSceneBindings& Bindings, TArray<FContextualAnimSetPivot>& OutPivots)
 {
 	Bindings.CalculateAnimSetPivots(OutPivots);
@@ -336,3 +352,17 @@ FTransform UContextualAnimUtilities::BP_SceneBindings_GetAlignmentTransformForRo
 	return Result;
 }
 
+FTransform UContextualAnimUtilities::BP_SceneBindings_GetAlignmentTransformForRoleRelativeToPivot(const FContextualAnimSceneBindings& Bindings, FName Role, const FContextualAnimSetPivot& Pivot, float Time)
+{
+	FTransform Result = FTransform::Identity;
+
+	if (const UContextualAnimSceneAsset* SceneAsset = Bindings.GetSceneAsset())
+	{
+		if (const FContextualAnimSceneSection* Section = SceneAsset->GetSection(Bindings.GetSectionIdx()))
+		{
+			return Section->GetAlignmentTransformForRoleRelativeToPivot(Bindings.GetAnimSetIdx(), Role, Time) * Pivot.Transform;
+		}
+	}
+
+	return Result;
+}
diff --git a/Engine/Plugins/Experimental/Animation/ContextualAnimation/Source/ContextualAnimation/Public/ContextualAnimSceneAsset.h b/Engine/Plugins/Experimental/Animation/ContextualAnimation/Source/ContextualAnimation/Public/ContextualAnimSceneAsset.h
--- a/Engine/Plugins/Experimental/Animation/ContextualAnimation/Source/ContextualAnimation/Public/ContextualAnimSceneAsset.h
+++ b/Engine/Plugins/Experimental/Animation/ContextualAnimation/Source/ContextualAnimation/Public/ContextualAnimSceneAsset.h
@@ -216,6 +216,9 @@ public:
 	UFUNCTION(BlueprintCallable, Category = "Contextual Anim|Scene Asset", meta = (DisplayName = "Get IK Target Transform For Role At Time"))
 	FTransform BP_GetIKTargetTransformForRoleAtTime(int32 SectionIdx, int32 AnimSetIdx, FName Role, FName TrackName, float Time) const;
 
+	UFUNCTION(BlueprintCallable, Category = "Contextual Anim|Scene Asset", meta = (DisplayName = "Get Start and End Time For Warp Section"))
+	void BP_GetStartAndEndTimeForWarpSection(int32 SectionIdx, int32 AnimSetIdx, FName Role, FName WarpSectionName, float& OutStartTime, float& OutEndTime) const;
+
 	//@TODO: Kept around only to do not break existing content. It will go away in the future.
 	UFUNCTION(BlueprintCallable, Category = "Contextual Anim|Scene Asset")
 	bool Query(FName Role, FContextualAnimQueryResult& OutResult, const FContextualAnimQueryParams& QueryParams, const FTransform& ToWorldTransform) const;
diff --git a/Engine/Plugins/Experimental/Animation/ContextualAnimation/Source/ContextualAnimation/Public/ContextualAnimTypes.h b/Engine/Plugins/Experimental/Animation/ContextualAnimation/Source/ContextualAnimation/Public/ContextualAnimTypes.h
--- a/Engine/Plugins/Experimental/Animation/ContextualAnimation/Source/ContextualAnimation/Public/ContextualAnimTypes.h
+++ b/Engine/Plugins/Experimental/Animation/ContextualAnimation/Source/ContextualAnimation/Public/ContextualAnimTypes.h
@@ -101,9 +101,13 @@ struct CONTEXTUALANIMATION_API FContextualAnimTrack
 	UPROPERTY(VisibleAnywhere, BlueprintReadOnly, Category = "Defaults")
 	int32 AnimTrackIdx = INDEX_NONE;
 
+	// DEPRECATED: Will go away soon
 	float GetSyncTimeForWarpSection(int32 WarpSectionIndex) const;
 	float GetSyncTimeForWarpSection(const FName& WarpSectionName) const;
 
+	void GetStartAndEndTimeForWarpSection(int32 WarpSectionIndex, float& OutStartTime, float& OutEndTime) const;
+	void GetStartAndEndTimeForWarpSection(const FName& WarpSectionName, float& OutStartTime, float& OutEndTime) const;
+
 	FORCEINLINE FTransform GetAlignmentTransformAtTime(float Time) const { return AlignmentData.ExtractTransformAtTime(0, Time); }
 	FORCEINLINE FTransform GetAlignmentTransformAtEntryTime() const { return AlignmentData.ExtractTransformAtTime(0, 0.f); }
 	FORCEINLINE FTransform GetAlignmentTransformAtSyncTime() const { return AlignmentData.ExtractTransformAtTime(0, GetSyncTimeForWarpSection(0)); }
diff --git a/Engine/Plugins/Experimental/Animation/ContextualAnimation/Source/ContextualAnimation/Public/ContextualAnimUtilities.h b/Engine/Plugins/Experimental/Animation/ContextualAnimation/Source/ContextualAnimation/Public/ContextualAnimUtilities.h
--- a/Engine/Plugins/Experimental/Animation/ContextualAnimation/Source/ContextualAnimation/Public/ContextualAnimUtilities.h
+++ b/Engine/Plugins/Experimental/Animation/ContextualAnimation/Source/ContextualAnimation/Public/ContextualAnimUtilities.h
@@ -74,7 +74,7 @@ public:
 
 	UFUNCTION(BlueprintCallable, Category = "Contextual Anim|Scene Bindings", meta = (DisplayName = "Calculate pivots For Bindings"))
 	static void BP_SceneBindings_CalculateAnimSetPivots(const FContextualAnimSceneBindings& Bindings, TArray<FContextualAnimSetPivot>& OutPivots);
-	
+
 	UFUNCTION(BlueprintCallable, Category = "Contextual Anim|Scene Bindings", meta = (DisplayName = "Add Or Update Warp Targets For Bindings"))
 	static void BP_SceneBindings_AddOrUpdateWarpTargetsForBindings(const FContextualAnimSceneBindings& Bindings);
 
@@ -84,12 +84,21 @@ public:
 	UFUNCTION(BlueprintPure, Category = "Contextual Anim|Scene Bindings", meta = (DisplayName = "Get Binding By Role"))
 	static const FContextualAnimSceneBinding& BP_SceneBindings_GetBindingByRole(const FContextualAnimSceneBindings& Bindings, FName Role);
 
+	UFUNCTION(BlueprintPure, Category = "Contextual Anim|Scene Bindings", meta = (DisplayName = "Get Binding By Actor"))
+	static const FContextualAnimSceneBinding& BP_SceneBindings_GetBindingByActor(const FContextualAnimSceneBindings& Bindings, const AActor* Actor);
+
 	UFUNCTION(BlueprintPure, Category = "Contextual Anim|Scene Bindings", meta = (DisplayName = "Get Scene Asset"))
 	static const UContextualAnimSceneAsset* BP_SceneBindings_GetSceneAsset(const FContextualAnimSceneBindings& Bindings) { return Bindings.GetSceneAsset(); }
 
+	UFUNCTION(BlueprintPure, Category = "Contextual Anim|Scene Bindings", meta = (DisplayName = "Get Section And AnimSet Indices"))
+	static void BP_SceneBindings_GetSectionAndAnimSetIndices(const FContextualAnimSceneBindings& Bindings, int32& SectionIdx, int32& AnimSetIdx);
+
 	UFUNCTION(BlueprintPure, Category = "Contextual Anim|Scene Bindings", meta = (DisplayName = "Get Alignment Transform For Role Relative To Other Role"))
 	static FTransform BP_SceneBindings_GetAlignmentTransformForRoleRelativeToOtherRole(const FContextualAnimSceneBindings& Bindings, FName Role, FName RelativeToRole, float Time);
 
+	UFUNCTION(BlueprintPure, Category = "Contextual Anim|Scene Bindings", meta = (DisplayName = "Get Alignment Transform For Role Relative To Pivot"))
+	static FTransform BP_SceneBindings_GetAlignmentTransformForRoleRelativeToPivot(const FContextualAnimSceneBindings& Bindings, FName Role, const FContextualAnimSetPivot& Pivot, float Time);
+
 	// FContextualAnimSceneBindingContext Blueprint Interface
 	//------------------------------------------------------------------------------------------
 	
@@ -99,6 +108,15 @@ public:
 	UFUNCTION(BlueprintPure, Category = "Contextual Anim|Scene Binding Context", meta = (DisplayName = "Make Contextual Anim Scene Binding Context With External Transform"))
 	static FContextualAnimSceneBindingContext BP_SceneBindingContext_MakeFromActorWithExternalTransform(AActor* Actor, FTransform ExternalTransform) { return FContextualAnimSceneBindingContext(Actor, ExternalTransform); }
 
+	UFUNCTION(BlueprintPure, Category = "Contextual Anim|Scene Binding Context", meta = (DisplayName = "Get Actor"))
+	static AActor* BP_SceneBindingContext_GetActor(const FContextualAnimSceneBindingContext& BindingContext) { return BindingContext.GetActor(); }
+
+	UFUNCTION(BlueprintPure, Category = "Contextual Anim|Scene Binding Context", meta = (DisplayName = "Get Transform"))
+	static FTransform BP_SceneBindingContext_GetTransform(const FContextualAnimSceneBindingContext& BindingContext) { return BindingContext.GetTransform(); }
+
+	UFUNCTION(BlueprintPure, Category = "Contextual Anim|Scene Binding Context", meta = (DisplayName = "Get Velocity"))
+	static FVector BP_SceneBindingContext_GetVelocity(const FContextualAnimSceneBindingContext& BindingContext) { return BindingContext.GetVelocity(); }
+
 	// FContextualAnimSceneBinding Blueprint Interface
 	//------------------------------------------------------------------------------------------
 
