Contextual Anim: Added selection criterion blueprintable base class so users can create game-specific criteria in blueprint

diff --git a/Engine/Plugins/Experimental/Animation/ContextualAnimation/Source/ContextualAnimation/Private/ContextualAnimSelectionCriterion.cpp b/Engine/Plugins/Experimental/Animation/ContextualAnimation/Source/ContextualAnimation/Private/ContextualAnimSelectionCriterion.cpp
--- a/Engine/Plugins/Experimental/Animation/ContextualAnimation/Source/ContextualAnimation/Private/ContextualAnimSelectionCriterion.cpp
+++ b/Engine/Plugins/Experimental/Animation/ContextualAnimation/Source/ContextualAnimation/Private/ContextualAnimSelectionCriterion.cpp
@@ -17,6 +17,23 @@ UContextualAnimSceneAsset* UContextualAnimSelectionCriterion::GetSceneAssetOwner
 	return Cast<UContextualAnimSceneAsset>(GetOuter());
 }
 
+// UContextualAnimSelectionCriterion_Blueprint
+//===========================================================================
+UContextualAnimSelectionCriterion_Blueprint::UContextualAnimSelectionCriterion_Blueprint(const FObjectInitializer& ObjectInitializer)
+	: Super(ObjectInitializer) 
+{
+}
+
+const UContextualAnimSceneAsset* UContextualAnimSelectionCriterion_Blueprint::GetSceneAsset() const
+{
+	return GetSceneAssetOwner();
+}
+
+bool UContextualAnimSelectionCriterion_Blueprint::DoesQuerierPassCondition(const FContextualAnimSceneBindingContext& Primary, const FContextualAnimSceneBindingContext& Querier) const
+{
+	return BP_DoesQuerierPassCondition(Primary, Querier);
+}
+
 // UContextualAnimSelectionCriterion_TriggerArea
 //===========================================================================
 
diff --git a/Engine/Plugins/Experimental/Animation/ContextualAnimation/Source/ContextualAnimation/Public/ContextualAnimSelectionCriterion.h b/Engine/Plugins/Experimental/Animation/ContextualAnimation/Source/ContextualAnimation/Public/ContextualAnimSelectionCriterion.h
--- a/Engine/Plugins/Experimental/Animation/ContextualAnimation/Source/ContextualAnimation/Public/ContextualAnimSelectionCriterion.h
+++ b/Engine/Plugins/Experimental/Animation/ContextualAnimation/Source/ContextualAnimation/Public/ContextualAnimSelectionCriterion.h
@@ -24,7 +24,7 @@ class CONTEXTUALANIMATION_API UContextualAnimSelectionCriterion : public UObject
 
 public:
 
-	UPROPERTY(EditAnywhere, BlueprintReadOnly, Category = "Defaults")
+	UPROPERTY(EditAnywhere, BlueprintReadOnly, Category = "Default")
 	EContextualAnimCriterionType Type = EContextualAnimCriterionType::Spatial;
 
 	UContextualAnimSelectionCriterion(const FObjectInitializer& ObjectInitializer);
@@ -34,6 +34,27 @@ public:
 	virtual bool DoesQuerierPassCondition(const FContextualAnimSceneBindingContext& Primary, const FContextualAnimSceneBindingContext& Querier) const { return false; }
 };
 
+// UContextualAnimSelectionCriterion_Blueprint
+//===========================================================================
+
+UCLASS(Abstract, Blueprintable)
+class CONTEXTUALANIMATION_API UContextualAnimSelectionCriterion_Blueprint : public UContextualAnimSelectionCriterion
+{
+	GENERATED_BODY()
+
+public:
+
+	UContextualAnimSelectionCriterion_Blueprint(const FObjectInitializer& ObjectInitializer);
+
+	UFUNCTION(BlueprintImplementableEvent, Category = "Default", meta = (DisplayName = "Does Querier Pass Condition"))
+	bool BP_DoesQuerierPassCondition(const FContextualAnimSceneBindingContext& Primary, const FContextualAnimSceneBindingContext& Querier) const;
+
+	UFUNCTION(BlueprintPure, Category = "Default")
+	const UContextualAnimSceneAsset* GetSceneAsset() const;
+
+	virtual bool DoesQuerierPassCondition(const FContextualAnimSceneBindingContext& Primary, const FContextualAnimSceneBindingContext& Querier) const override;
+};
+
 // UContextualAnimSelectionCriterion_TriggerArea
 //===========================================================================
 
@@ -44,10 +65,10 @@ class CONTEXTUALANIMATION_API UContextualAnimSelectionCriterion_TriggerArea : pu
 
 public:
 
-	UPROPERTY(EditAnywhere, EditFixedSize, BlueprintReadWrite, Category = "Defaults", meta = (EditFixedOrder))
+	UPROPERTY(EditAnywhere, EditFixedSize, BlueprintReadWrite, Category = "Default", meta = (EditFixedOrder))
 	TArray<FVector> PolygonPoints;
 
-	UPROPERTY(EditAnywhere, BlueprintReadWrite, Category = "Defaults", meta = (ClampMin = "0", UIMin = "0"))
+	UPROPERTY(EditAnywhere, BlueprintReadWrite, Category = "Default", meta = (ClampMin = "0", UIMin = "0"))
 	float Height = 100.f;
 
 	UContextualAnimSelectionCriterion_TriggerArea(const FObjectInitializer& ObjectInitializer);
@@ -75,16 +96,16 @@ class CONTEXTUALANIMATION_API UContextualAnimSelectionCriterion_Angle : public U
 
 public:
 
-	UPROPERTY(EditAnywhere, BlueprintReadWrite, Category = "Defaults")
+	UPROPERTY(EditAnywhere, BlueprintReadWrite, Category = "Default")
 	EContextualAnimCriterionAngleMode Mode = EContextualAnimCriterionAngleMode::ToPrimary;
 
-	UPROPERTY(EditAnywhere, BlueprintReadWrite, Category = "Defaults")
+	UPROPERTY(EditAnywhere, BlueprintReadWrite, Category = "Default")
 	bool bUseSignedAngle = false;
 
-	UPROPERTY(EditAnywhere, BlueprintReadWrite, Category = "Defaults", meta = (ClampMin = "-180", UIMin = "-180", ClampMax = "180", UIMax = "180"))
+	UPROPERTY(EditAnywhere, BlueprintReadWrite, Category = "Default", meta = (ClampMin = "-180", UIMin = "-180", ClampMax = "180", UIMax = "180"))
 	float MinAngle = 0.f;
 
-	UPROPERTY(EditAnywhere, BlueprintReadWrite, Category = "Defaults", meta = (ClampMin = "-180", UIMin = "-180", ClampMax = "180", UIMax = "180"))
+	UPROPERTY(EditAnywhere, BlueprintReadWrite, Category = "Default", meta = (ClampMin = "-180", UIMin = "-180", ClampMax = "180", UIMax = "180"))
 	float MaxAngle = 0.f;
 
 	UContextualAnimSelectionCriterion_Angle(const FObjectInitializer& ObjectInitializer) : Super(ObjectInitializer) {}
@@ -109,13 +130,13 @@ class CONTEXTUALANIMATION_API UContextualAnimSelectionCriterion_Distance : publi
 
 public:
 
-	UPROPERTY(EditAnywhere, BlueprintReadWrite, Category = "Defaults")
+	UPROPERTY(EditAnywhere, BlueprintReadWrite, Category = "Default")
 	EContextualAnimCriterionDistanceMode Mode = EContextualAnimCriterionDistanceMode::Distance_2D;
 
-	UPROPERTY(EditAnywhere, BlueprintReadWrite, Category = "Defaults", meta = (ClampMin = "0", UIMin = "0"))
+	UPROPERTY(EditAnywhere, BlueprintReadWrite, Category = "Default", meta = (ClampMin = "0", UIMin = "0"))
 	float MinDistance = 0.f;
 
-	UPROPERTY(EditAnywhere, BlueprintReadWrite, Category = "Defaults", meta = (ClampMin = "0", UIMin = "0"))
+	UPROPERTY(EditAnywhere, BlueprintReadWrite, Category = "Default", meta = (ClampMin = "0", UIMin = "0"))
 	float MaxDistance = 0.f;
 
 	UContextualAnimSelectionCriterion_Distance(const FObjectInitializer& ObjectInitializer) : Super(ObjectInitializer) {}
