Contextual Anim Editor: Highlight viewport border when simulate mode is active

diff --git a/Engine/Plugins/Experimental/Animation/ContextualAnimation/Source/ContextualAnimationEditor/Private/ContextualAnimEditorStyle.cpp b/Engine/Plugins/Experimental/Animation/ContextualAnimation/Source/ContextualAnimationEditor/Private/ContextualAnimEditorStyle.cpp
--- a/Engine/Plugins/Experimental/Animation/ContextualAnimation/Source/ContextualAnimationEditor/Private/ContextualAnimEditorStyle.cpp
+++ b/Engine/Plugins/Experimental/Animation/ContextualAnimation/Source/ContextualAnimationEditor/Private/ContextualAnimEditorStyle.cpp
@@ -47,6 +47,10 @@ TSharedRef< FSlateStyleSet > FContextualAnimEditorStyle::Create()
 
 	Style->Set("ContextualAnimEditor.Icon", new IMAGE_BRUSH(TEXT("ButtonIcon_40x"), Icon40x40));
 
+	const FString EngineEditorSlateDir = FPaths::EngineContentDir() / TEXT("Editor/Slate");
+	Style->SetContentRoot(EngineEditorSlateDir);
+	Style->Set("ContextualAnimEditor.Viewport.Border", new BOX_BRUSH("Old/Window/ViewportDebugBorder", 0.8f, FLinearColor(1.f, 1.f, 1.f, 1.f)));
+
 	return Style;
 }
 
diff --git a/Engine/Plugins/Experimental/Animation/ContextualAnimation/Source/ContextualAnimationEditor/Private/SContextualAnimViewport.cpp b/Engine/Plugins/Experimental/Animation/ContextualAnimation/Source/ContextualAnimationEditor/Private/SContextualAnimViewport.cpp
--- a/Engine/Plugins/Experimental/Animation/ContextualAnimation/Source/ContextualAnimationEditor/Private/SContextualAnimViewport.cpp
+++ b/Engine/Plugins/Experimental/Animation/ContextualAnimation/Source/ContextualAnimationEditor/Private/SContextualAnimViewport.cpp
@@ -2,9 +2,10 @@
 
 #include "SContextualAnimViewport.h"
 #include "ContextualAnimViewportClient.h"
-#include "Toolkits/AssetEditorToolkit.h"
 #include "SContextualAnimViewportToolbar.h"
 #include "ContextualAnimAssetEditorCommands.h"
+#include "ContextualAnimEditorStyle.h"
+#include "ContextualAnimAssetEditorToolkit.h"
 
 void SContextualAnimViewport::Construct(const FArguments& InArgs, const FContextualAnimViewportRequiredArgs& InRequiredArgs)
 {
@@ -18,6 +19,22 @@ void SContextualAnimViewport::Construct(const FArguments& InArgs, const FContext
 	);
 }
 
+const FSlateBrush* SContextualAnimViewport::OnGetViewportBorderBrush() const
+{
+	// Highlight the border of the viewport when Simulate Mode is active
+	if (AssetEditorToolkitPtr.Pin()->IsSimulateModeActive())
+	{
+		return FContextualAnimEditorStyle::Get().GetBrush("ContextualAnimEditor.Viewport.Border");
+	}
+
+	return nullptr;
+}
+
+FSlateColor SContextualAnimViewport::OnGetViewportBorderColorAndOpacity() const
+{
+	return FLinearColor::Yellow;
+}
+
 void SContextualAnimViewport::BindCommands()
 {
 	SEditorViewport::BindCommands();
diff --git a/Engine/Plugins/Experimental/Animation/ContextualAnimation/Source/ContextualAnimationEditor/Public/SContextualAnimViewport.h b/Engine/Plugins/Experimental/Animation/ContextualAnimation/Source/ContextualAnimationEditor/Public/SContextualAnimViewport.h
--- a/Engine/Plugins/Experimental/Animation/ContextualAnimation/Source/ContextualAnimationEditor/Public/SContextualAnimViewport.h
+++ b/Engine/Plugins/Experimental/Animation/ContextualAnimation/Source/ContextualAnimationEditor/Public/SContextualAnimViewport.h
@@ -45,6 +45,8 @@ protected:
 	virtual void BindCommands() override;
 	virtual TSharedRef<FEditorViewportClient> MakeEditorViewportClient() override;
 	virtual TSharedPtr<SWidget> MakeViewportToolbar() override;
+	virtual const FSlateBrush* OnGetViewportBorderBrush() const override;
+	virtual FSlateColor OnGetViewportBorderColorAndOpacity() const;
 	// ~End of SEditorViewport interface
 
 	/** The viewport toolbar */
