Contextual Anim Editor: Refactor IsSimulateModeActive flag into Inactive, Paused and Playing states

diff --git a/Engine/Plugins/Experimental/Animation/ContextualAnimation/Source/ContextualAnimationEditor/Private/ContextualAnimAssetEditorToolkit.cpp b/Engine/Plugins/Experimental/Animation/ContextualAnimation/Source/ContextualAnimationEditor/Private/ContextualAnimAssetEditorToolkit.cpp
--- a/Engine/Plugins/Experimental/Animation/ContextualAnimation/Source/ContextualAnimationEditor/Private/ContextualAnimAssetEditorToolkit.cpp
+++ b/Engine/Plugins/Experimental/Animation/ContextualAnimation/Source/ContextualAnimationEditor/Private/ContextualAnimAssetEditorToolkit.cpp
@@ -180,7 +180,7 @@ void FContextualAnimAssetEditorToolkit::ToggleSimulateMode()
 
 bool FContextualAnimAssetEditorToolkit::IsSimulateModeActive() const
 {
-	return (ViewModel.IsValid()) ? ViewModel->IsSimulateModeActive() : false;
+	return (ViewModel.IsValid() && !ViewModel->IsSimulateModeInactive());
 }
 
 void FContextualAnimAssetEditorToolkit::ExtendToolbar()
diff --git a/Engine/Plugins/Experimental/Animation/ContextualAnimation/Source/ContextualAnimationEditor/Private/ContextualAnimEdMode.cpp b/Engine/Plugins/Experimental/Animation/ContextualAnimation/Source/ContextualAnimationEditor/Private/ContextualAnimEdMode.cpp
--- a/Engine/Plugins/Experimental/Animation/ContextualAnimation/Source/ContextualAnimationEditor/Private/ContextualAnimEdMode.cpp
+++ b/Engine/Plugins/Experimental/Animation/ContextualAnimation/Source/ContextualAnimationEditor/Private/ContextualAnimEdMode.cpp
@@ -229,7 +229,7 @@ bool FContextualAnimEdMode::InputDelta(FEditorViewportClient* InViewportClient,
 
 bool FContextualAnimEdMode::InputKey(FEditorViewportClient* ViewportClient, FViewport* Viewport, FKey Key, EInputEvent Event)
 {
-	if(ViewModel && ViewModel->IsSimulateModeActive())
+	if(ViewModel && ViewModel->IsSimulateModePaused())
 	{
 		if (Key == EKeys::Enter && Event == IE_Released)
 		{
diff --git a/Engine/Plugins/Experimental/Animation/ContextualAnimation/Source/ContextualAnimationEditor/Private/ContextualAnimViewModel.cpp b/Engine/Plugins/Experimental/Animation/ContextualAnimation/Source/ContextualAnimationEditor/Private/ContextualAnimViewModel.cpp
--- a/Engine/Plugins/Experimental/Animation/ContextualAnimation/Source/ContextualAnimationEditor/Private/ContextualAnimViewModel.cpp
+++ b/Engine/Plugins/Experimental/Animation/ContextualAnimation/Source/ContextualAnimationEditor/Private/ContextualAnimViewModel.cpp
@@ -473,12 +473,12 @@ void FContextualAnimViewModel::AddNewIKTarget(const UContextualAnimNewIKTargetPa
 	SceneAsset->MarkPackageDirty();
 }
 
-void FContextualAnimViewModel::ToggleSimulateMode() 
-{ 
-	bIsSimulateModeActive = !bIsSimulateModeActive; 
-
-	if(bIsSimulateModeActive)
+void FContextualAnimViewModel::ToggleSimulateMode()
+{
+	if (SimulateModeState == ESimulateModeState::Inactive)
 	{
+		SimulateModeState = ESimulateModeState::Paused;
+
 		if (SceneInstance.IsValid())
 		{
 			SceneInstance->Stop();
@@ -486,7 +486,9 @@ void FContextualAnimViewModel::ToggleSimulateMode()
 	}
 	else
 	{
-		if(SceneInstance.IsValid())
+		SimulateModeState = ESimulateModeState::Inactive;
+
+		if (SceneInstance.IsValid())
 		{
 			for (auto& Binding : SceneInstance->GetBindings())
 			{
@@ -517,6 +519,8 @@ void FContextualAnimViewModel::StartSimulation()
 		UE_LOG(LogContextualAnim, Warning, TEXT("Can't start scene"));
 		return;
 	}
+
+	SimulateModeState = ESimulateModeState::Playing;
 }
 
 UWorld* FContextualAnimViewModel::GetWorld() const
@@ -756,15 +760,16 @@ FText FContextualAnimViewModel::GetSelectionDebugText() const
 
 bool FContextualAnimViewModel::ProcessInputDelta(FVector& InDrag, FRotator& InRot, FVector& InScale)
 {
-	if(IsSimulateModeActive())
+	if (IsSimulateModePaused())
 	{
+		// In Simulate Mode (while paused) the user can drag the preview actor around to simulate the interaction from different start positions
 		if (AActor* SelectedActor = GetSelectedActor())
 		{
 			SelectedActor->SetActorLocationAndRotation(SelectedActor->GetActorLocation() + InDrag, SelectedActor->GetActorRotation() + InRot);
 			return true;
 		}
 	}
-	else
+	else if (IsSimulateModeInactive())
 	{
 		if (UContextualAnimSelectionCriterion_TriggerArea* Spatial = Cast<UContextualAnimSelectionCriterion_TriggerArea>(GetSelectedSelectionCriterion()))
 		{
@@ -837,14 +842,16 @@ bool FContextualAnimViewModel::GetCustomDrawingCoordinateSystem(FMatrix& InMatri
 
 bool FContextualAnimViewModel::ShouldPreviewSceneDrawWidget() const
 {
-	const UContextualAnimSelectionCriterion* SelectionCriterion = GetSelectedSelectionCriterion();
-	if (SelectionCriterion && SelectionCriterion->GetClass()->IsChildOf<UContextualAnimSelectionCriterion_TriggerArea>())
+	// When Simulate Mode is inactive we show the widget if an editable selection criterion is selected (only Trigger Area for now) or if an actor is selected
+	if (IsSimulateModeInactive())
 	{
-		return true;
+		const UContextualAnimSelectionCriterion* SelectionCriterion = GetSelectedSelectionCriterion();
+		return ((SelectionCriterion && SelectionCriterion->GetClass()->IsChildOf<UContextualAnimSelectionCriterion_TriggerArea>()) || GetSelectedActor());
 	}
-	else if (GetSelectedActor())
+	// When Simulate Mode is Paused we show the widget if an actor is selected, so the user can modify the position of any of the actor before triggering the interaction
+	else if (IsSimulateModePaused())
 	{
-		return true;
+		return GetSelectedActor() != nullptr;
 	}
 
 	return false;
diff --git a/Engine/Plugins/Experimental/Animation/ContextualAnimation/Source/ContextualAnimationEditor/Public/ContextualAnimViewModel.h b/Engine/Plugins/Experimental/Animation/ContextualAnimation/Source/ContextualAnimationEditor/Public/ContextualAnimViewModel.h
--- a/Engine/Plugins/Experimental/Animation/ContextualAnimation/Source/ContextualAnimationEditor/Public/ContextualAnimViewModel.h
+++ b/Engine/Plugins/Experimental/Animation/ContextualAnimation/Source/ContextualAnimationEditor/Public/ContextualAnimViewModel.h
@@ -39,6 +39,13 @@ public:
 		Notifies
 	};
 
+	enum class ESimulateModeState : uint8
+	{
+		Inactive,	// Simulate Mode Inactive
+		Paused,		// Simulate Mode Active but interaction is not playing
+		Playing		// Simulate Mode Active and interaction playing
+	};
+
 	FContextualAnimViewModel();
 	virtual ~FContextualAnimViewModel();
 
@@ -76,7 +83,9 @@ public:
 	void SetActiveAnimSetForSection(int32 SectionIdx, int32 AnimSetIdx);
 
 	void ToggleSimulateMode();
-	bool IsSimulateModeActive() { return bIsSimulateModeActive; }
+	bool IsSimulateModeInactive()	const { return SimulateModeState == ESimulateModeState::Inactive; }
+	bool IsSimulateModePaused()		const { return SimulateModeState == ESimulateModeState::Paused; }
+	bool IsSimulateModePlaying()	const { return SimulateModeState == ESimulateModeState::Playing; }
 	void StartSimulation();
 
 	void UpdatePreviewActorTransform(const FContextualAnimSceneBinding& Binding, float Time);
@@ -157,7 +166,7 @@ private:
 	/** Flag for preventing OnAnimNotifyChanged from updating tracks when the change to the animation came from us */
 	bool bUpdatingAnimationFromSequencer = false;
 
-	bool bIsSimulateModeActive = false;
+	ESimulateModeState SimulateModeState = ESimulateModeState::Inactive;
 
 	FContextualAnimStartSceneParams StartSceneParams;
 
