Contextual Anim Editor: Refactored the way users can modify the transform of an actor in the scene

diff --git a/Engine/Plugins/Experimental/Animation/ContextualAnimation/Source/ContextualAnimationEditor/Private/ContextualAnimEdMode.cpp b/Engine/Plugins/Experimental/Animation/ContextualAnimation/Source/ContextualAnimationEditor/Private/ContextualAnimEdMode.cpp
--- a/Engine/Plugins/Experimental/Animation/ContextualAnimation/Source/ContextualAnimationEditor/Private/ContextualAnimEdMode.cpp
+++ b/Engine/Plugins/Experimental/Animation/ContextualAnimation/Source/ContextualAnimationEditor/Private/ContextualAnimEdMode.cpp
@@ -204,11 +204,21 @@ bool FContextualAnimEdMode::GetHitResultUnderCursor(FHitResult& OutHitResult, FE
 
 bool FContextualAnimEdMode::StartTracking(FEditorViewportClient* InViewportClient, FViewport* InViewport)
 {
+	if (ViewModel)
+	{
+		return ViewModel->StartTracking();
+	}
+
 	return FEdMode::StartTracking(InViewportClient, InViewport);
 }
 
 bool FContextualAnimEdMode::EndTracking(FEditorViewportClient* InViewportClient, FViewport* InViewport)
 {
+	if (ViewModel)
+	{
+		return ViewModel->EndTracking();
+	}
+
 	return FEdMode::EndTracking(InViewportClient, InViewport);
 }
 
diff --git a/Engine/Plugins/Experimental/Animation/ContextualAnimation/Source/ContextualAnimationEditor/Private/ContextualAnimViewModel.cpp b/Engine/Plugins/Experimental/Animation/ContextualAnimation/Source/ContextualAnimationEditor/Private/ContextualAnimViewModel.cpp
--- a/Engine/Plugins/Experimental/Animation/ContextualAnimation/Source/ContextualAnimationEditor/Private/ContextualAnimViewModel.cpp
+++ b/Engine/Plugins/Experimental/Animation/ContextualAnimation/Source/ContextualAnimationEditor/Private/ContextualAnimViewModel.cpp
@@ -33,6 +33,9 @@
 #include "Animation/DebugSkelMeshComponent.h"
 #include "AnimPreviewInstance.h"
 #include "ContextualAnimSelectionCriterion.h"
+#include "ScopedTransaction.h"
+
+#define LOCTEXT_NAMESPACE "ContextualAnimViewModel"
 
 FContextualAnimViewModel::FContextualAnimViewModel()
 	: SceneAsset(nullptr)
@@ -504,6 +507,8 @@ void FContextualAnimViewModel::ToggleSimulateMode()
 
 		SetDefaultMode();
 	}
+
+	ClearSelection();
 };
 
 void FContextualAnimViewModel::StartSimulation()
@@ -538,7 +543,23 @@ void FContextualAnimViewModel::UpdatePreviewActorTransform(const FContextualAnim
 {
 	if (AActor* PreviewActor = Binding.GetActor())
 	{
-		FTransform Transform = Binding.GetAnimTrack().GetRootTransformAtTime(Time);
+		FTransform Transform = FTransform::Identity;
+
+		// When modifying the actor transform, use the 'temp' transform instead of the one saved in the asset
+		// This allows the user to scrub the time line while maintaining the new MeshToScene transform 
+		if(ModifyingActorTransformInSceneState != EModifyActorTransformInSceneState::Inactive && Binding.GetActor() == ModifyingTransformInSceneCachedActor.Get())
+		{
+			if (Binding.GetAnimTrack().Animation)
+			{
+				Transform = UContextualAnimUtilities::ExtractRootTransformFromAnimation(Binding.GetAnimTrack().Animation, Time);
+			}
+
+			Transform *= NewMeshToSceneTransform;
+		}
+		else
+		{
+			Transform = Binding.GetAnimTrack().GetRootTransformAtTime(Time);
+		}
 
 		// Special case for Character
 		if (ACharacter* PreviewCharacter = Cast<ACharacter>(PreviewActor))
@@ -698,6 +719,13 @@ void FContextualAnimViewModel::UpdateSelection(const AActor* SelectedActor)
 	const FContextualAnimSceneBinding* Binding = SceneInstance.IsValid() ? SceneInstance->FindBindingByActor(SelectedActor) : nullptr;
 	if (Binding)
 	{
+		// Discard changes if we were modifying the transform of an actor in the scene but never confirmed the changes before selecting a different actor
+		if(ModifyingActorTransformInSceneState == EModifyActorTransformInSceneState::WaitingForConfirmation && 
+			Binding->GetActor() != ModifyingTransformInSceneCachedActor.Get())
+		{
+			DiscardChangeToActorTransformInScene();
+		}
+
 		UpdateSelection(Binding->GetRoleDef().Name);
 	}
 	else
@@ -801,13 +829,17 @@ bool FContextualAnimViewModel::ProcessInputDelta(FVector& InDrag, FRotator& InRo
 
 			return true;
 		}
-		else if (FContextualAnimSceneBinding* Binding = GetSelectedBinding())
+		else if (ModifyingActorTransformInSceneState == EModifyActorTransformInSceneState::Modifying)
 		{
-			FTransform& MeshToSpaceTransform = (const_cast<FContextualAnimTrack*>(&Binding->GetAnimTrack()))->MeshToScene;
-			MeshToSpaceTransform.SetLocation(MeshToSpaceTransform.GetLocation() + InDrag);
-			MeshToSpaceTransform.SetRotation(FQuat(InRot) * MeshToSpaceTransform.GetRotation());
-
-			UpdatePreviewActorTransform(*Binding, Sequencer->GetGlobalTime().AsSeconds());
+			if (WantsToModifyMeshToSceneForSelectedActor())
+			{
+				if (AActor* SelectedActor = GetSelectedActor())
+				{
+					NewMeshToSceneTransform.SetLocation(NewMeshToSceneTransform.GetLocation() + InDrag);
+					NewMeshToSceneTransform.SetRotation(FQuat(InRot) * NewMeshToSceneTransform.GetRotation());
+					SelectedActor->SetActorLocationAndRotation(SelectedActor->GetActorLocation() + InDrag, SelectedActor->GetActorRotation() + InRot);
+				}
+			}
 
 			return true;
 		}
@@ -840,13 +872,21 @@ bool FContextualAnimViewModel::GetCustomDrawingCoordinateSystem(FMatrix& InMatri
 	return false;
 }
 
+bool FContextualAnimViewModel::WantsToModifyMeshToSceneForSelectedActor() const
+{
+	//@TODO: Communicate this to the user somehow
+	return FSlateApplication::Get().GetModifierKeys().IsShiftDown();
+}
+
 bool FContextualAnimViewModel::ShouldPreviewSceneDrawWidget() const
 {
-	// When Simulate Mode is inactive we show the widget if an editable selection criterion is selected (only Trigger Area for now) or if an actor is selected
+	// When Simulate Mode is inactive we show the widget if an editable selection criterion is selected (only Trigger Area for now) 
+	// or if an actor is selected and user wants to modify the MeshToScene transform
 	if (IsSimulateModeInactive())
 	{
 		const UContextualAnimSelectionCriterion* SelectionCriterion = GetSelectedSelectionCriterion();
-		return ((SelectionCriterion && SelectionCriterion->GetClass()->IsChildOf<UContextualAnimSelectionCriterion_TriggerArea>()) || GetSelectedActor());
+		return ((SelectionCriterion && SelectionCriterion->GetClass()->IsChildOf<UContextualAnimSelectionCriterion_TriggerArea>()) ||
+			(WantsToModifyMeshToSceneForSelectedActor() && GetSelectedActor()));
 	}
 	// When Simulate Mode is Paused we show the widget if an actor is selected, so the user can modify the position of any of the actor before triggering the interaction
 	else if (IsSimulateModePaused())
@@ -890,6 +930,79 @@ FVector FContextualAnimViewModel::GetWidgetLocationFromSelection() const
 	return FVector::ZeroVector;
 }
 
+bool FContextualAnimViewModel::StartTracking()
+{
+	if (IsSimulateModeInactive() && WantsToModifyMeshToSceneForSelectedActor())
+	{
+		if (const FContextualAnimSceneBinding* Binding = GetSelectedBinding())
+		{
+			NewMeshToSceneTransform = Binding->GetAnimTrack().MeshToScene;
+			ModifyingActorTransformInSceneState = EModifyActorTransformInSceneState::Modifying;
+			ModifyingTransformInSceneCachedActor = Binding->GetActor();
+			return true;
+		}
+	}
+
+	return false;
+}
+
+bool FContextualAnimViewModel::EndTracking()
+{
+	if (ModifyingActorTransformInSceneState == EModifyActorTransformInSceneState::Modifying)
+	{
+		const FContextualAnimSceneBinding* Binding = GetSelectedBinding();
+		if (Binding && !Binding->GetAnimTrack().MeshToScene.Equals(NewMeshToSceneTransform))
+		{
+			ModifyingActorTransformInSceneState = EModifyActorTransformInSceneState::WaitingForConfirmation;
+		}
+		else
+		{
+			ModifyingActorTransformInSceneState = EModifyActorTransformInSceneState::Inactive;
+			ModifyingTransformInSceneCachedActor.Reset();
+		}
+
+		return true;		
+	}
+
+	return false;
+}
+
+bool FContextualAnimViewModel::IsChangeToActorTransformInSceneWaitingForConfirmation() const
+{
+	const FContextualAnimSceneBinding* Binding = GetSelectedBinding();
+	return (Binding && Binding->GetActor() == ModifyingTransformInSceneCachedActor.Get() && ModifyingActorTransformInSceneState == EModifyActorTransformInSceneState::WaitingForConfirmation);
+}
+
+void FContextualAnimViewModel::ApplyChangeToActorTransformInScene()
+{
+	const FContextualAnimSceneBinding* Binding = GetSelectedBinding();
+	if(Binding && Binding->GetActor() == ModifyingTransformInSceneCachedActor.Get())
+	{
+		FScopedTransaction Transaction(LOCTEXT("ModifyMeshToSceneTransform", "Modify MeshToScene Transform"));
+		SceneAsset->Modify();
+
+		FTransform& MeshToSpaceTransform = (const_cast<FContextualAnimTrack*>(&Binding->GetAnimTrack()))->MeshToScene;
+		MeshToSpaceTransform = NewMeshToSceneTransform;
+
+		UpdatePreviewActorTransform(*Binding, Sequencer->GetGlobalTime().AsSeconds());
+	}
+	
+	ModifyingActorTransformInSceneState = EModifyActorTransformInSceneState::Inactive;
+	ModifyingTransformInSceneCachedActor.Reset();
+}
+
+void FContextualAnimViewModel::DiscardChangeToActorTransformInScene()
+{
+	ModifyingActorTransformInSceneState = EModifyActorTransformInSceneState::Inactive;
+
+	if (const FContextualAnimSceneBinding* Binding = SceneInstance->FindBindingByActor(ModifyingTransformInSceneCachedActor.Get()))
+	{
+		UpdatePreviewActorTransform(*Binding, Sequencer->GetGlobalTime().AsSeconds());
+	}
+
+	ModifyingTransformInSceneCachedActor.Reset();
+}
+
 void FContextualAnimViewModel::OnFinishedChangingProperties(const FPropertyChangedEvent& PropertyChangedEvent)
 {
 	const FName PropertyName = (PropertyChangedEvent.Property != nullptr) ? PropertyChangedEvent.Property->GetFName() : NAME_None;
@@ -902,4 +1015,41 @@ void FContextualAnimViewModel::OnFinishedChangingProperties(const FPropertyChang
 	{
 		SetDefaultMode();
 	}
-}
\ No newline at end of file
+}
+
+bool FContextualAnimViewModel::MatchesContext(const FTransactionContext& InContext, const TArray<TPair<UObject*, FTransactionObjectEvent>>& TransactionObjectContexts) const
+{
+	// Ensure that we only react to modifications to the SceneAsset
+	if (SceneAsset)
+	{
+		for (const TPair<UObject*, FTransactionObjectEvent>& TransactionObjectPair : TransactionObjectContexts)
+		{
+			UObject* Object = TransactionObjectPair.Key;
+			while (Object != nullptr)
+			{
+				if (Object == SceneAsset)
+				{
+					return true;
+				}
+
+				Object = Object->GetOuter();
+			}
+		}
+	}
+
+	return false;
+}
+
+void FContextualAnimViewModel::PostUndo(bool bSuccess)
+{
+	// Refresh everything after a Undo operation
+	SetDefaultMode();
+}
+
+void FContextualAnimViewModel::PostRedo(bool bSuccess)
+{
+	// Refresh everything after a Redo operation
+	SetDefaultMode();
+}
+
+#undef LOCTEXT_NAMESPACE
\ No newline at end of file
diff --git a/Engine/Plugins/Experimental/Animation/ContextualAnimation/Source/ContextualAnimationEditor/Private/SContextualAnimViewport.cpp b/Engine/Plugins/Experimental/Animation/ContextualAnimation/Source/ContextualAnimationEditor/Private/SContextualAnimViewport.cpp
--- a/Engine/Plugins/Experimental/Animation/ContextualAnimation/Source/ContextualAnimationEditor/Private/SContextualAnimViewport.cpp
+++ b/Engine/Plugins/Experimental/Animation/ContextualAnimation/Source/ContextualAnimationEditor/Private/SContextualAnimViewport.cpp
@@ -6,6 +6,12 @@
 #include "ContextualAnimAssetEditorCommands.h"
 #include "ContextualAnimEditorStyle.h"
 #include "ContextualAnimAssetEditorToolkit.h"
+#include "ContextualAnimViewModel.h"
+#include "Styling/AppStyle.h"
+#include "EditorFontGlyphs.h"
+#include "Widgets/Input/SButton.h"
+
+#define LOCTEXT_NAMESPACE "ContextualAnimViewport"
 
 void SContextualAnimViewport::Construct(const FArguments& InArgs, const FContextualAnimViewportRequiredArgs& InRequiredArgs)
 {
@@ -17,6 +23,137 @@ void SContextualAnimViewport::Construct(const FArguments& InArgs, const FContext
 		.IsEnabled(FSlateApplication::Get().GetNormalExecutionAttribute())
 		.AddMetaData<FTagMetaData>(TEXT("AnimationTools.Viewport"))
 	);
+
+	auto GetNotificationText = [this]()
+	{
+		return LOCTEXT("MeshToSceneChangeConfirmText", "Apply changes to the transform of the actor in the scene?");
+	};
+
+	auto GetNotificationVisibility = [this]()
+	{
+		return AssetEditorToolkitPtr.Pin()->GetViewModel()->IsChangeToActorTransformInSceneWaitingForConfirmation() ? EVisibility::Visible : EVisibility::Hidden;
+	};
+
+	auto ApplyChanges = [this]()
+	{
+		AssetEditorToolkitPtr.Pin()->GetViewModel()->ApplyChangeToActorTransformInScene();
+		return FReply::Handled();
+	};
+
+	auto CancelChanges = [this]()
+ 	{
+		AssetEditorToolkitPtr.Pin()->GetViewModel()->DiscardChangeToActorTransformInScene();
+		return FReply::Handled();
+	};
+
+	ViewportOverlay->AddSlot()
+	[
+		SNew(SVerticalBox)
+		.Visibility_Lambda(GetNotificationVisibility)
+		+ SVerticalBox::Slot()
+		.FillHeight(1)
+		[
+			SNew(SOverlay)
+			+ SOverlay::Slot()
+			.Padding(8)
+			.VAlign(VAlign_Bottom)
+			.HAlign(HAlign_Right)
+			[
+				SNew(SBorder)
+				.BorderImage(FAppStyle::GetBrush("AnimViewport.Notification.Warning"))
+				[
+					SNew(SHorizontalBox)
+					+ SHorizontalBox::Slot()
+					.FillWidth(1.0f)
+					.Padding(4.0f, 4.0f)
+					[
+						SNew(SHorizontalBox)
+						.ToolTipText_Lambda(GetNotificationText)
+						+ SHorizontalBox::Slot()
+						.AutoWidth()
+						.VAlign(VAlign_Center)
+						.Padding(0.0f, 0.0f, 4.0f, 0.0f)
+						[
+							SNew(STextBlock)
+							.TextStyle(FAppStyle::Get(), "AnimViewport.MessageText")
+							.Font(FAppStyle::Get().GetFontStyle("FontAwesome.9"))
+							.Text(FEditorFontGlyphs::Exclamation_Triangle)
+						]
+						+ SHorizontalBox::Slot()
+						.VAlign(VAlign_Center)
+						.FillWidth(1.0f)
+						[
+							SNew(STextBlock)
+							.Text_Lambda(GetNotificationText)
+							.TextStyle(FAppStyle::Get(), "AnimViewport.MessageText")
+						]
+					]
+					+ SHorizontalBox::Slot()
+					.AutoWidth()
+					.Padding(2.0f, 0.0f)
+					[
+						SNew(SButton)
+						.ForegroundColor(FSlateColor::UseForeground())
+						.ButtonStyle(FAppStyle::Get(), "FlatButton.Success")
+						.ToolTipText_Lambda(GetNotificationText)
+						.OnClicked_Lambda(ApplyChanges)
+						[
+							SNew(SHorizontalBox)
+							+ SHorizontalBox::Slot()
+							.AutoWidth()
+							.VAlign(VAlign_Center)
+							.Padding(0.0f, 0.0f, 4.0f, 0.0f)
+							[
+								SNew(STextBlock)
+								.TextStyle(FAppStyle::Get(), "AnimViewport.MessageText")
+								.Font(FAppStyle::Get().GetFontStyle("FontAwesome.9"))
+								.Text(FEditorFontGlyphs::Pencil)
+							]
+							+ SHorizontalBox::Slot()
+							.VAlign(VAlign_Center)
+							.AutoWidth()
+							[
+								SNew(STextBlock)
+								.TextStyle(FAppStyle::Get(), "AnimViewport.MessageText")
+								.Text(LOCTEXT("ApplyButtonText", "Apply"))
+							]
+						]
+					]
+					+ SHorizontalBox::Slot()
+					.AutoWidth()
+					.Padding(2.0f, 0.0f)
+					[
+						SNew(SButton)
+						.ForegroundColor(FSlateColor::UseForeground())
+						.ButtonStyle(FAppStyle::Get(), "FlatButton.Danger")
+						.ToolTipText(LOCTEXT("DiscardMeshToSceneChangeToolTip", "Discard changes to the transform of the actor in the scene."))
+						.OnClicked_Lambda(CancelChanges)
+						[
+							SNew(SHorizontalBox)
+							+ SHorizontalBox::Slot()
+							.AutoWidth()
+							.VAlign(VAlign_Center)
+							.Padding(0.0f, 0.0f, 4.0f, 0.0f)
+							[
+								SNew(STextBlock)
+								.TextStyle(FAppStyle::Get(), "AnimViewport.MessageText")
+								.Font(FAppStyle::Get().GetFontStyle("FontAwesome.9"))
+								.Text(FEditorFontGlyphs::Times)
+							]
+							+ SHorizontalBox::Slot()
+							.VAlign(VAlign_Center)
+							.AutoWidth()
+							[
+								SNew(STextBlock)
+								.TextStyle(FAppStyle::Get(), "AnimViewport.MessageText")
+								.Text(LOCTEXT("CancelButtonText", "Cancel"))
+							]
+						]
+					]
+				]
+			]
+		]
+	];
 }
 
 const FSlateBrush* SContextualAnimViewport::OnGetViewportBorderBrush() const
@@ -91,4 +228,6 @@ TSharedPtr<FExtender> SContextualAnimViewport::GetExtenders() const
 
 void SContextualAnimViewport::OnFloatingButtonClicked()
 {
-}
\ No newline at end of file
+}
+
+#undef LOCTEXT_NAMESPACE
\ No newline at end of file
diff --git a/Engine/Plugins/Experimental/Animation/ContextualAnimation/Source/ContextualAnimationEditor/Public/ContextualAnimViewModel.h b/Engine/Plugins/Experimental/Animation/ContextualAnimation/Source/ContextualAnimationEditor/Public/ContextualAnimViewModel.h
--- a/Engine/Plugins/Experimental/Animation/ContextualAnimation/Source/ContextualAnimationEditor/Public/ContextualAnimViewModel.h
+++ b/Engine/Plugins/Experimental/Animation/ContextualAnimation/Source/ContextualAnimationEditor/Public/ContextualAnimViewModel.h
@@ -8,6 +8,7 @@
 #include "ContextualAnimTypes.h"
 #include "ContextualAnimMovieSceneSequence.h"
 #include "ISequencer.h"
+#include "EditorUndoClient.h"
 
 class UWorld;
 class FContextualAnimPreviewScene;
@@ -29,7 +30,7 @@ struct FMovieSceneSectionMovedParams;
 struct FContextualAnimNewAnimSetParams;
 struct FContextualAnimTrack;
 
-class FContextualAnimViewModel : public TSharedFromThis<FContextualAnimViewModel>, public FGCObject
+class FContextualAnimViewModel : public TSharedFromThis<FContextualAnimViewModel>, public FGCObject, public FSelfRegisteringEditorUndoClient
 {
 public:
 
@@ -53,6 +54,11 @@ public:
 	virtual void AddReferencedObjects(FReferenceCollector& Collector) override;
 	virtual FString GetReferencerName() const override { return TEXT("FContextualAnimViewModel"); }
 
+	// ~ FEditorUndoClient
+	virtual bool MatchesContext(const FTransactionContext& InContext, const TArray<TPair<UObject*, FTransactionObjectEvent>>& TransactionObjectContexts) const override;
+	virtual void PostUndo(bool bSuccess) override;
+	virtual void PostRedo(bool bSuccess) override;
+
 	void OnFinishedChangingProperties(const FPropertyChangedEvent& PropertyChangedEvent);
 
 	void Initialize(UContextualAnimSceneAsset* InSceneAsset, const TSharedRef<FContextualAnimPreviewScene>& InPreviewScene);
@@ -103,6 +109,12 @@ public:
 	bool ShouldPreviewSceneDrawWidget() const;
 	bool GetCustomDrawingCoordinateSystem(FMatrix& InMatrix, void* InData);
 	FVector GetWidgetLocationFromSelection() const;
+	bool StartTracking();
+	bool EndTracking();
+
+	bool IsChangeToActorTransformInSceneWaitingForConfirmation() const;
+	void ApplyChangeToActorTransformInScene();
+	void DiscardChangeToActorTransformInScene();
 
 private:
 
@@ -174,6 +186,16 @@ private:
 
 	TWeakObjectPtr<UAnimSequenceBase> EditingAnimation = nullptr;
 
+	/** Enum to track the current state when modifying the transform of an actor in the scene */
+	enum class EModifyActorTransformInSceneState : uint8 {  Inactive, Modifying, WaitingForConfirmation };
+	EModifyActorTransformInSceneState ModifyingActorTransformInSceneState = EModifyActorTransformInSceneState::Inactive;
+
+	/** New MeshToScene value for the selected actor. Updated while dragging the actor around. Only committed when the user confirms the operation */
+	FTransform NewMeshToSceneTransform = FTransform::Identity;
+
+	/** Selected actor when the user starts modifying its transform in the scene */
+	TWeakObjectPtr<AActor> ModifyingTransformInSceneCachedActor;
+
 	AActor* SpawnPreviewActor(const FContextualAnimTrack& AnimTrack);
 
 	UWorld* GetWorld() const;
@@ -187,4 +209,6 @@ private:
 	void OnAnimNotifyChanged(UAnimSequenceBase* Animation);
 
 	void CreateSequencer();
+
+	bool WantsToModifyMeshToSceneForSelectedActor() const;
 };
\ No newline at end of file
