Adding additional includes in preparation for removing a global include

#rnx

diff --git a/Engine/Plugins/Animation/ControlRig/Source/ControlRigDeveloper/Private/ControlRigBlueprint.cpp b/Engine/Plugins/Animation/ControlRig/Source/ControlRigDeveloper/Private/ControlRigBlueprint.cpp
--- a/Engine/Plugins/Animation/ControlRig/Source/ControlRigDeveloper/Private/ControlRigBlueprint.cpp
+++ b/Engine/Plugins/Animation/ControlRig/Source/ControlRigDeveloper/Private/ControlRigBlueprint.cpp
@@ -42,6 +42,7 @@
 #include "Settings/ControlRigSettings.h"
 #include "UnrealEdGlobals.h"
 #include "Editor/UnrealEdEngine.h"
+#include "Editor/Transactor.h"
 #include "CookOnTheSide/CookOnTheFlyServer.h"
 #include "Widgets/Notifications/SNotificationList.h"
 #include "Framework/Notifications/NotificationManager.h"
diff --git a/Engine/Plugins/Animation/ControlRig/Source/ControlRigDeveloper/Private/Graph/ControlRigGraphSchema.cpp b/Engine/Plugins/Animation/ControlRig/Source/ControlRigDeveloper/Private/Graph/ControlRigGraphSchema.cpp
--- a/Engine/Plugins/Animation/ControlRig/Source/ControlRigDeveloper/Private/Graph/ControlRigGraphSchema.cpp
+++ b/Engine/Plugins/Animation/ControlRig/Source/ControlRigDeveloper/Private/Graph/ControlRigGraphSchema.cpp
@@ -36,6 +36,7 @@
 #if WITH_EDITOR
 #include "ControlRigEditor/Private/Editor/SControlRigFunctionLocalizationWidget.h"
 #include "Misc/MessageDialog.h"
+#include "Editor/Transactor.h"
 #endif
 
 #define LOCTEXT_NAMESPACE "ControlRigGraphSchema"
diff --git a/Engine/Plugins/Animation/ControlRig/Source/ControlRigEditor/Private/Editor/ControlRigEditor.cpp b/Engine/Plugins/Animation/ControlRig/Source/ControlRigEditor/Private/Editor/ControlRigEditor.cpp
--- a/Engine/Plugins/Animation/ControlRig/Source/ControlRigEditor/Private/Editor/ControlRigEditor.cpp
+++ b/Engine/Plugins/Animation/ControlRig/Source/ControlRigEditor/Private/Editor/ControlRigEditor.cpp
@@ -10,6 +10,7 @@
 #include "SEnumCombo.h"
 #include "Framework/Commands/GenericCommands.h"
 #include "Editor.h"
+#include "Editor/Transactor.h"
 #include "Graph/ControlRigGraph.h"
 #include "BlueprintActionDatabase.h"
 #include "ControlRigBlueprintCommands.h"
diff --git a/Engine/Plugins/Animation/ControlRig/Source/ControlRigEditor/Private/Sequencer/ControlRigParameterTrackEditor.cpp b/Engine/Plugins/Animation/ControlRig/Source/ControlRigEditor/Private/Sequencer/ControlRigParameterTrackEditor.cpp
--- a/Engine/Plugins/Animation/ControlRig/Source/ControlRigEditor/Private/Sequencer/ControlRigParameterTrackEditor.cpp
+++ b/Engine/Plugins/Animation/ControlRig/Source/ControlRigEditor/Private/Sequencer/ControlRigParameterTrackEditor.cpp
@@ -78,6 +78,7 @@
 #include "ControlRigSpaceChannelEditors.h"
 #include "TransformConstraint.h"
 #include "Misc/ScopedSlowTask.h"
+#include "Misc/TransactionObjectEvent.h"
 
 #define LOCTEXT_NAMESPACE "FControlRigParameterTrackEditor"
 
diff --git a/Engine/Plugins/Developer/TraceSourceFiltering/Source/SourceFilteringEditor/Private/EditorSessionSourceFilterService.cpp b/Engine/Plugins/Developer/TraceSourceFiltering/Source/SourceFilteringEditor/Private/EditorSessionSourceFilterService.cpp
--- a/Engine/Plugins/Developer/TraceSourceFiltering/Source/SourceFilteringEditor/Private/EditorSessionSourceFilterService.cpp
+++ b/Engine/Plugins/Developer/TraceSourceFiltering/Source/SourceFilteringEditor/Private/EditorSessionSourceFilterService.cpp
@@ -8,6 +8,7 @@
 #include "UObject/SoftObjectPath.h"
 #include "ScopedTransaction.h"
 #include "Editor.h"
+#include "Misc/TransactionObjectEvent.h"
 #include "Factories/BlueprintFactory.h"
 
 #include "TraceSourceFiltering.h"
diff --git a/Engine/Plugins/Enterprise/DataprepEditor/Source/DataprepCore/Private/DataprepActionAsset.cpp b/Engine/Plugins/Enterprise/DataprepEditor/Source/DataprepCore/Private/DataprepActionAsset.cpp
--- a/Engine/Plugins/Enterprise/DataprepEditor/Source/DataprepCore/Private/DataprepActionAsset.cpp
+++ b/Engine/Plugins/Enterprise/DataprepEditor/Source/DataprepCore/Private/DataprepActionAsset.cpp
@@ -30,6 +30,7 @@
 #include "UObject/UObjectHash.h"
 
 #include "Editor.h"
+#include "Misc/TransactionObjectEvent.h"
 
 void UDataprepActionStep::PostLoad()
 {
diff --git a/Engine/Plugins/Enterprise/DataprepEditor/Source/DataprepCore/Private/Parameterization/DataprepParameterization.cpp b/Engine/Plugins/Enterprise/DataprepEditor/Source/DataprepCore/Private/Parameterization/DataprepParameterization.cpp
--- a/Engine/Plugins/Enterprise/DataprepEditor/Source/DataprepCore/Private/Parameterization/DataprepParameterization.cpp
+++ b/Engine/Plugins/Enterprise/DataprepEditor/Source/DataprepCore/Private/Parameterization/DataprepParameterization.cpp
@@ -17,6 +17,7 @@
 #include "Engine/Engine.h"
 #include "Math/UnrealMathUtility.h"
 #include "Misc/ITransaction.h"
+#include "Misc/TransactionObjectEvent.h"
 #include "Serialization/Archive.h"
 #include "Templates/UnrealTemplate.h"
 #include "UObject/CoreObjectVersion.h"
diff --git a/Engine/Plugins/Enterprise/VariantManager/Source/VariantManager/Private/SVariantManager.cpp b/Engine/Plugins/Enterprise/VariantManager/Source/VariantManager/Private/SVariantManager.cpp
--- a/Engine/Plugins/Enterprise/VariantManager/Source/VariantManager/Private/SVariantManager.cpp
+++ b/Engine/Plugins/Enterprise/VariantManager/Source/VariantManager/Private/SVariantManager.cpp
@@ -56,7 +56,7 @@
 #include "Misc/CString.h"
 #include "Misc/ConfigCacheIni.h"
 #include "Misc/FileHelper.h"
-#include "Misc/ITransaction.h"
+#include "Misc/TransactionObjectEvent.h"
 #include "Misc/Paths.h"
 #include "Modules/ModuleManager.h"
 #include "ObjectTools.h"
diff --git a/Engine/Plugins/Experimental/Animation/ContextualAnimation/Source/ContextualAnimationEditor/Private/ContextualAnimViewModel.cpp b/Engine/Plugins/Experimental/Animation/ContextualAnimation/Source/ContextualAnimationEditor/Private/ContextualAnimViewModel.cpp
--- a/Engine/Plugins/Experimental/Animation/ContextualAnimation/Source/ContextualAnimationEditor/Private/ContextualAnimViewModel.cpp
+++ b/Engine/Plugins/Experimental/Animation/ContextualAnimation/Source/ContextualAnimationEditor/Private/ContextualAnimViewModel.cpp
@@ -16,6 +16,7 @@
 #include "Animation/AnimInstance.h"
 #include "Animation/AnimNotifies/AnimNotifyState.h"
 #include "EngineUtils.h"
+#include "Misc/TransactionObjectEvent.h"
 #include "ContextualAnimTypes.h"
 #include "ContextualAnimEditorTypes.h"
 #include "ContextualAnimUtilities.h"
diff --git a/Engine/Plugins/Experimental/VirtualCamera/Source/VirtualCamera/Private/VCamBlueprintFunctionLibrary.cpp b/Engine/Plugins/Experimental/VirtualCamera/Source/VirtualCamera/Private/VCamBlueprintFunctionLibrary.cpp
--- a/Engine/Plugins/Experimental/VirtualCamera/Source/VirtualCamera/Private/VCamBlueprintFunctionLibrary.cpp
+++ b/Engine/Plugins/Experimental/VirtualCamera/Source/VirtualCamera/Private/VCamBlueprintFunctionLibrary.cpp
@@ -10,6 +10,7 @@
 #include "GameFramework/PlayerController.h"
 #if WITH_EDITOR
 #include "Editor.h"
+#include "Editor/Transactor.h"
 #include "EditorAssetLibrary.h"
 #include "ITakeRecorderModule.h"
 #include "LevelEditor.h"
diff --git a/Engine/Plugins/FX/Niagara/Source/NiagaraEditor/Private/Toolkits/NiagaraScriptToolkit.cpp b/Engine/Plugins/FX/Niagara/Source/NiagaraEditor/Private/Toolkits/NiagaraScriptToolkit.cpp
--- a/Engine/Plugins/FX/Niagara/Source/NiagaraEditor/Private/Toolkits/NiagaraScriptToolkit.cpp
+++ b/Engine/Plugins/FX/Niagara/Source/NiagaraEditor/Private/Toolkits/NiagaraScriptToolkit.cpp
@@ -33,6 +33,7 @@
 #include "Framework/MultiBox/MultiBoxBuilder.h"
 #include "Misc/FeedbackContext.h"
 #include "Misc/MessageDialog.h"
+#include "Misc/TransactionObjectEvent.h"
 #include "Modules/ModuleManager.h"
 #include "UObject/Package.h"
 #include "ViewModels/NiagaraParameterPanelViewModel.h"
diff --git a/Engine/Plugins/FX/Niagara/Source/NiagaraEditor/Private/Toolkits/NiagaraSystemToolkit.cpp b/Engine/Plugins/FX/Niagara/Source/NiagaraEditor/Private/Toolkits/NiagaraSystemToolkit.cpp
--- a/Engine/Plugins/FX/Niagara/Source/NiagaraEditor/Private/Toolkits/NiagaraSystemToolkit.cpp
+++ b/Engine/Plugins/FX/Niagara/Source/NiagaraEditor/Private/Toolkits/NiagaraSystemToolkit.cpp
@@ -44,6 +44,7 @@
 #include "Widgets/SNiagaraEmitterVersionWidget.h"
 #include "BusyCursor.h"
 #include "Misc/FeedbackContext.h"
+#include "Misc/TransactionObjectEvent.h"
 #include "Editor.h"
 #include "LevelEditor.h"
 #include "LevelEditorActions.h"
diff --git a/Engine/Plugins/Runtime/nDisplay/Source/DisplayClusterLightCardEditor/Private/SDisplayClusterLightCardEditor.cpp b/Engine/Plugins/Runtime/nDisplay/Source/DisplayClusterLightCardEditor/Private/SDisplayClusterLightCardEditor.cpp
--- a/Engine/Plugins/Runtime/nDisplay/Source/DisplayClusterLightCardEditor/Private/SDisplayClusterLightCardEditor.cpp
+++ b/Engine/Plugins/Runtime/nDisplay/Source/DisplayClusterLightCardEditor/Private/SDisplayClusterLightCardEditor.cpp
@@ -31,6 +31,7 @@
 #include "Misc/FileHelper.h"
 #include "PropertyCustomizationHelpers.h"
 #include "ScopedTransaction.h"
+#include "Misc/TransactionObjectEvent.h"
 #include "Selection.h"
 #include "Toolkits/AssetEditorToolkit.h"
 #include "Widgets/Docking/SDockTab.h"
diff --git a/Engine/Plugins/Tests/EditorTests/Source/EditorTests/Private/UnrealEd/BlueprintAutomationTests.cpp b/Engine/Plugins/Tests/EditorTests/Source/EditorTests/Private/UnrealEd/BlueprintAutomationTests.cpp
--- a/Engine/Plugins/Tests/EditorTests/Source/EditorTests/Private/UnrealEd/BlueprintAutomationTests.cpp
+++ b/Engine/Plugins/Tests/EditorTests/Source/EditorTests/Private/UnrealEd/BlueprintAutomationTests.cpp
@@ -30,6 +30,7 @@
 #include "GameFramework/SaveGame.h"
 #include "EngineUtils.h"
 #include "Editor.h"
+#include "Editor/Transactor.h"
 #include "FileHelpers.h"
 
 #include "ObjectTools.h"
diff --git a/Engine/Plugins/Tests/EditorTests/Source/EditorTests/Private/UnrealEd/EditorAutomationTests.cpp b/Engine/Plugins/Tests/EditorTests/Source/EditorTests/Private/UnrealEd/EditorAutomationTests.cpp
--- a/Engine/Plugins/Tests/EditorTests/Source/EditorTests/Private/UnrealEd/EditorAutomationTests.cpp
+++ b/Engine/Plugins/Tests/EditorTests/Source/EditorTests/Private/UnrealEd/EditorAutomationTests.cpp
@@ -48,6 +48,7 @@
 #include "Tests/AutomationEditorCommon.h"
 
 #include "PackageTools.h"
+#include "Editor/Transactor.h"
 
 #include "K2Node_Event.h"
 #include "K2Node_CustomEvent.h"
diff --git a/Engine/Plugins/Tests/EditorTests/Source/EditorTests/Private/UnrealEd/FbxAutomationTests.cpp b/Engine/Plugins/Tests/EditorTests/Source/EditorTests/Private/UnrealEd/FbxAutomationTests.cpp
--- a/Engine/Plugins/Tests/EditorTests/Source/EditorTests/Private/UnrealEd/FbxAutomationTests.cpp
+++ b/Engine/Plugins/Tests/EditorTests/Source/EditorTests/Private/UnrealEd/FbxAutomationTests.cpp
@@ -40,6 +40,7 @@
 #include "FbxMeshUtils.h"
 #include "Tests/FbxAutomationCommon.h"
 
+#include "Editor/Transactor.h"
 #include "Editor/EditorEngine.h"
 extern UNREALED_API UEditorEngine* GEditor;
 
diff --git a/Engine/Plugins/Tests/EditorTests/Source/EditorTests/Private/UnrealEd/ObjectToolsTests.cpp b/Engine/Plugins/Tests/EditorTests/Source/EditorTests/Private/UnrealEd/ObjectToolsTests.cpp
--- a/Engine/Plugins/Tests/EditorTests/Source/EditorTests/Private/UnrealEd/ObjectToolsTests.cpp
+++ b/Engine/Plugins/Tests/EditorTests/Source/EditorTests/Private/UnrealEd/ObjectToolsTests.cpp
@@ -5,6 +5,7 @@
 #include "Misc/Paths.h"
 #include "Misc/AutomationTest.h"
 #include "Editor.h"
+#include "Editor/Transactor.h"
 #include "ObjectTools.h"
 #include "ScopedTransaction.h"
 #include "Tests/AutomationTestSettings.h"
diff --git a/Engine/Plugins/Tests/InterchangeTests/Source/InterchangeTests/Private/InterchangeImportTest.cpp b/Engine/Plugins/Tests/InterchangeTests/Source/InterchangeTests/Private/InterchangeImportTest.cpp
--- a/Engine/Plugins/Tests/InterchangeTests/Source/InterchangeTests/Private/InterchangeImportTest.cpp
+++ b/Engine/Plugins/Tests/InterchangeTests/Source/InterchangeTests/Private/InterchangeImportTest.cpp
@@ -17,6 +17,7 @@
 #include "AssetRegistry/AssetRegistryModule.h"
 #include "AssetRegistry/AssetData.h"
 #include "Editor.h"
+#include "Editor/Transactor.h"
 #include "InterchangeManager.h"
 #include "UObject/SavePackage.h"
 #include "ObjectTools.h"
diff --git a/Engine/Plugins/VirtualProduction/RemoteControl/Source/RemoteControl/Private/RemoteControlModule.cpp b/Engine/Plugins/VirtualProduction/RemoteControl/Source/RemoteControl/Private/RemoteControlModule.cpp
--- a/Engine/Plugins/VirtualProduction/RemoteControl/Source/RemoteControl/Private/RemoteControlModule.cpp
+++ b/Engine/Plugins/VirtualProduction/RemoteControl/Source/RemoteControl/Private/RemoteControlModule.cpp
@@ -30,6 +30,7 @@
 #if WITH_EDITOR
 #include "Editor.h"
 #include "Editor/EditorEngine.h"
+#include "Editor/Transactor.h"
 #include "HAL/IConsoleManager.h"
 #include "ScopedTransaction.h"
 #endif
diff --git a/Engine/Source/Editor/Kismet/Private/BlueprintCompilationManager.cpp b/Engine/Source/Editor/Kismet/Private/BlueprintCompilationManager.cpp
--- a/Engine/Source/Editor/Kismet/Private/BlueprintCompilationManager.cpp
+++ b/Engine/Source/Editor/Kismet/Private/BlueprintCompilationManager.cpp
@@ -11,6 +11,7 @@
 #include "Editor.h"
 #include "Editor/UnrealEdEngine.h"
 #include "Engine/Engine.h"
+#include "Editor/Transactor.h"
 #include "Engine/LevelScriptBlueprint.h"
 #include "Engine/SCS_Node.h"
 #include "Engine/SimpleConstructionScript.h"
diff --git a/Engine/Source/Editor/Kismet/Private/BlueprintEditor.cpp b/Engine/Source/Editor/Kismet/Private/BlueprintEditor.cpp
--- a/Engine/Source/Editor/Kismet/Private/BlueprintEditor.cpp
+++ b/Engine/Source/Editor/Kismet/Private/BlueprintEditor.cpp
@@ -33,6 +33,7 @@
 #include "EdGraph/EdGraphSchema.h"
 #include "EdGraphNode_Comment.h"
 #include "Editor/UnrealEdEngine.h"
+#include "Editor/Transactor.h"
 #include "Settings/EditorExperimentalSettings.h"
 #include "Settings/BlueprintEditorProjectSettings.h"
 #include "GeneralProjectSettings.h"
diff --git a/Engine/Source/Editor/Kismet/Private/BlueprintEditorModule.cpp b/Engine/Source/Editor/Kismet/Private/BlueprintEditorModule.cpp
--- a/Engine/Source/Editor/Kismet/Private/BlueprintEditorModule.cpp
+++ b/Engine/Source/Editor/Kismet/Private/BlueprintEditorModule.cpp
@@ -15,6 +15,7 @@
 #include "Kismet2/KismetEditorUtilities.h"
 #include "Kismet2/KismetDebugUtilities.h"
 #include "Kismet2/BlueprintEditorUtils.h"
+#include "Editor/Transactor.h"
 #include "Editor/LevelEditor/Public/LevelEditor.h"
 #include "UserDefinedEnumEditor.h"
 #include "MessageLogInitializationOptions.h"
diff --git a/Engine/Source/Editor/LevelEditor/Private/SLevelViewport.cpp b/Engine/Source/Editor/LevelEditor/Private/SLevelViewport.cpp
--- a/Engine/Source/Editor/LevelEditor/Private/SLevelViewport.cpp
+++ b/Engine/Source/Editor/LevelEditor/Private/SLevelViewport.cpp
@@ -30,6 +30,7 @@
 #include "Styling/CoreStyle.h"
 #include "Styling/AppStyle.h"
 #include "Editor/UnrealEdEngine.h"
+#include "Editor/Transactor.h"
 #include "Exporters/ExportTextContainer.h"
 #include "Camera/CameraActor.h"
 #include "Camera/CameraComponent.h"
diff --git a/Engine/Source/Editor/MainFrame/Private/Frame/MainFrameActions.cpp b/Engine/Source/Editor/MainFrame/Private/Frame/MainFrameActions.cpp
--- a/Engine/Source/Editor/MainFrame/Private/Frame/MainFrameActions.cpp
+++ b/Engine/Source/Editor/MainFrame/Private/Frame/MainFrameActions.cpp
@@ -34,6 +34,7 @@
 #include "CookerSettings.h"
 #include "UnrealEdMisc.h"
 #include "Editor/UnrealEdEngine.h"
+#include "Editor/Transactor.h"
 #include "Preferences/UnrealEdOptions.h"
 #include "UnrealEdGlobals.h"
 #include "FileHelpers.h"
diff --git a/Engine/Source/Editor/MainFrame/Private/Menus/MainMenu.cpp b/Engine/Source/Editor/MainFrame/Private/Menus/MainMenu.cpp
--- a/Engine/Source/Editor/MainFrame/Private/Menus/MainMenu.cpp
+++ b/Engine/Source/Editor/MainFrame/Private/Menus/MainMenu.cpp
@@ -16,6 +16,7 @@
 #include "Styling/AppStyle.h"
 #include "Settings/EditorStyleSettings.h"
 #include "Editor/UnrealEdEngine.h"
+#include "Editor/Transactor.h"
 #include "Settings/EditorExperimentalSettings.h"
 #include "UnrealEdGlobals.h"
 #include "Frame/MainFrameActions.h"
diff --git a/Engine/Source/Editor/Sequencer/Private/Sequencer.cpp b/Engine/Source/Editor/Sequencer/Private/Sequencer.cpp
--- a/Engine/Source/Editor/Sequencer/Private/Sequencer.cpp
+++ b/Engine/Source/Editor/Sequencer/Private/Sequencer.cpp
@@ -29,6 +29,7 @@
 #include "Misc/MessageDialog.h"
 #include "Misc/FeedbackContext.h"
 #include "Misc/ScopedSlowTask.h"
+#include "Misc/TransactionObjectEvent.h"
 #include "Modules/ModuleManager.h"
 #include "UObject/UObjectIterator.h"
 #include "UObject/MetaData.h"
diff --git a/Engine/Source/Editor/UnrealEd/Private/AssetDeleteModel.cpp b/Engine/Source/Editor/UnrealEd/Private/AssetDeleteModel.cpp
--- a/Engine/Source/Editor/UnrealEd/Private/AssetDeleteModel.cpp
+++ b/Engine/Source/Editor/UnrealEd/Private/AssetDeleteModel.cpp
@@ -20,6 +20,7 @@
 #include "Settings/EditorLoadingSavingSettings.h"
 #include "EngineGlobals.h"
 #include "Editor.h"
+#include "Editor/Transactor.h"
 #include "FileHelpers.h"
 #include "UnrealEdGlobals.h"
 
diff --git a/Engine/Source/Editor/UnrealEd/Private/Dialogs/SDeleteAssetsDialog.cpp b/Engine/Source/Editor/UnrealEd/Private/Dialogs/SDeleteAssetsDialog.cpp
--- a/Engine/Source/Editor/UnrealEd/Private/Dialogs/SDeleteAssetsDialog.cpp
+++ b/Engine/Source/Editor/UnrealEd/Private/Dialogs/SDeleteAssetsDialog.cpp
@@ -18,6 +18,7 @@
 #include "IContentBrowserSingleton.h"
 #include "ContentBrowserModule.h"
 #include "Editor.h"
+#include "Editor/Transactor.h"
 
 #include "Framework/Commands/GenericCommands.h"
 #include "Subsystems/AssetEditorSubsystem.h"
diff --git a/Engine/Source/Editor/UnrealEd/Private/EditorEngine.cpp b/Engine/Source/Editor/UnrealEd/Private/EditorEngine.cpp
--- a/Engine/Source/Editor/UnrealEd/Private/EditorEngine.cpp
+++ b/Engine/Source/Editor/UnrealEd/Private/EditorEngine.cpp
@@ -71,6 +71,7 @@
 #include "UnrealEngine.h"
 #include "EngineUtils.h"
 #include "Editor.h"
+#include "Editor/Transactor.h"
 #include "EditorViewportClient.h"
 #include "LevelEditorViewport.h"
 #include "EditorModeManager.h"
diff --git a/Engine/Source/Editor/UnrealEd/Private/EditorLevelUtils.cpp b/Engine/Source/Editor/UnrealEd/Private/EditorLevelUtils.cpp
--- a/Engine/Source/Editor/UnrealEd/Private/EditorLevelUtils.cpp
+++ b/Engine/Source/Editor/UnrealEd/Private/EditorLevelUtils.cpp
@@ -27,6 +27,7 @@ EditorLevelUtils.cpp: Editor-specific level management routines
 #include "Engine/LevelStreaming.h"
 #include "Engine/Selection.h"
 #include "Editor.h"
+#include "Editor/Transactor.h"
 #include "EditorModeManager.h"
 #include "EditorModes.h"
 #include "FileHelpers.h"
diff --git a/Engine/Source/Editor/UnrealEd/Private/Kismet2/Kismet2.cpp b/Engine/Source/Editor/UnrealEd/Private/Kismet2/Kismet2.cpp
--- a/Engine/Source/Editor/UnrealEd/Private/Kismet2/Kismet2.cpp
+++ b/Engine/Source/Editor/UnrealEd/Private/Kismet2/Kismet2.cpp
@@ -57,6 +57,7 @@
 #include "Kismet2/KismetDebugUtilities.h"
 #include "Framework/MultiBox/MultiBoxBuilder.h"
 #include "Editor.h"
+#include "Editor/Transactor.h"
 
 #include "Editor/Kismet/Public/BlueprintEditorModule.h"
 #include "Editor/Kismet/Public/FindInBlueprintManager.h"
diff --git a/Engine/Source/Editor/UnrealEd/Private/ObjectTools.cpp b/Engine/Source/Editor/UnrealEd/Private/ObjectTools.cpp
--- a/Engine/Source/Editor/UnrealEd/Private/ObjectTools.cpp
+++ b/Engine/Source/Editor/UnrealEd/Private/ObjectTools.cpp
@@ -57,6 +57,7 @@
 #include "Engine/UserDefinedStruct.h"
 #include "Animation/MorphTarget.h"
 #include "Editor.h"
+#include "Editor/Transactor.h"
 #include "EditorDirectories.h"
 #include "FileHelpers.h"
 #include "Dialogs/Dialogs.h"
diff --git a/Engine/Source/Editor/UnrealEd/Private/PlayLevel.cpp b/Engine/Source/Editor/UnrealEd/Private/PlayLevel.cpp
--- a/Engine/Source/Editor/UnrealEd/Private/PlayLevel.cpp
+++ b/Engine/Source/Editor/UnrealEd/Private/PlayLevel.cpp
@@ -56,6 +56,7 @@
 #include "UnrealEngine.h"
 #include "EngineUtils.h"
 #include "Editor.h"
+#include "Editor/Transactor.h"
 #include "LevelEditorViewport.h"
 #include "EditorModeManager.h"
 #include "EditorModes.h"
diff --git a/Engine/Source/Editor/UnrealEd/Private/Subsystems/EditorAssetSubsystem.cpp b/Engine/Source/Editor/UnrealEd/Private/Subsystems/EditorAssetSubsystem.cpp
--- a/Engine/Source/Editor/UnrealEd/Private/Subsystems/EditorAssetSubsystem.cpp
+++ b/Engine/Source/Editor/UnrealEd/Private/Subsystems/EditorAssetSubsystem.cpp
@@ -6,6 +6,7 @@
 #include "AssetRegistry/AssetRegistryModule.h"
 #include "AssetToolsModule.h"
 #include "Editor.h"
+#include "Editor/Transactor.h"
 #include "EditorScriptingHelpers.h"
 #include "Experimental/Async/LazyEvent.h"
 #include "FileHelpers.h"
diff --git a/Engine/Source/Editor/WorldBrowser/Private/Tiles/WorldTileModel.cpp b/Engine/Source/Editor/WorldBrowser/Private/Tiles/WorldTileModel.cpp
--- a/Engine/Source/Editor/WorldBrowser/Private/Tiles/WorldTileModel.cpp
+++ b/Engine/Source/Editor/WorldBrowser/Private/Tiles/WorldTileModel.cpp
@@ -8,6 +8,7 @@
 #include "Engine/LevelBounds.h"
 #include "Engine/LevelStreamingDynamic.h"
 #include "Editor.h"
+#include "Editor/Transactor.h"
 #include "ScopedTransaction.h"
 #include "EditorLevelUtils.h"
 #include "LevelCollectionModel.h"
diff --git a/Engine/Source/Runtime/AIModule/Private/EnvironmentQuery/EQSTestingPawn.cpp b/Engine/Source/Runtime/AIModule/Private/EnvironmentQuery/EQSTestingPawn.cpp
--- a/Engine/Source/Runtime/AIModule/Private/EnvironmentQuery/EQSTestingPawn.cpp
+++ b/Engine/Source/Runtime/AIModule/Private/EnvironmentQuery/EQSTestingPawn.cpp
@@ -17,6 +17,7 @@
 #endif // WITH_EDITORONLY_DATA
 
 #if WITH_EDITOR
+#include "Misc/TransactionObjectEvent.h"
 #include "Editor/EditorEngine.h"
 extern UNREALED_API UEditorEngine* GEditor;
 #endif // WITH_EDITOR
diff --git a/Engine/Source/Runtime/Engine/Private/ActorEditor.cpp b/Engine/Source/Runtime/Engine/Private/ActorEditor.cpp
--- a/Engine/Source/Runtime/Engine/Private/ActorEditor.cpp
+++ b/Engine/Source/Runtime/Engine/Private/ActorEditor.cpp
@@ -29,6 +29,7 @@
 #if WITH_EDITOR
 
 #include "Editor.h"
+#include "Misc/TransactionObjectEvent.h"
 #include "Engine/LevelStreaming.h"
 #include "WorldPartition/WorldPartitionActorDesc.h"
 #include "WorldPartition/DataLayer/IDataLayerEditorModule.h"
diff --git a/Engine/Source/Runtime/Engine/Private/EdGraph/EdGraphPin.cpp b/Engine/Source/Runtime/Engine/Private/EdGraph/EdGraphPin.cpp
--- a/Engine/Source/Runtime/Engine/Private/EdGraph/EdGraphPin.cpp
+++ b/Engine/Source/Runtime/Engine/Private/EdGraph/EdGraphPin.cpp
@@ -16,6 +16,7 @@
 #include "HAL/IConsoleManager.h"
 #if WITH_EDITOR
 #include "Editor/EditorEngine.h"
+#include "Editor/Transactor.h"
 #include "Misc/ConfigCacheIni.h"
 #include "TickableEditorObject.h"
 #endif
diff --git a/Engine/Source/Runtime/Engine/Private/WorldSettings.cpp b/Engine/Source/Runtime/Engine/Private/WorldSettings.cpp
--- a/Engine/Source/Runtime/Engine/Private/WorldSettings.cpp
+++ b/Engine/Source/Runtime/Engine/Private/WorldSettings.cpp
@@ -36,6 +36,7 @@
 
 #if WITH_EDITOR
 #include "Editor.h"
+#include "Misc/TransactionObjectEvent.h"
 #include "HierarchicalLOD.h"
 #include "Settings/EditorExperimentalSettings.h"
 #include "Landscape.h"
