ContextualAnim: Refactor selection criterion based on angle into a cone and added debug visualization for it in the preview scene.

diff --git a/Engine/Plugins/Experimental/Animation/ContextualAnimation/Source/ContextualAnimation/Private/ContextualAnimSelectionCriterion.cpp b/Engine/Plugins/Experimental/Animation/ContextualAnimation/Source/ContextualAnimation/Private/ContextualAnimSelectionCriterion.cpp
--- a/Engine/Plugins/Experimental/Animation/ContextualAnimation/Source/ContextualAnimation/Private/ContextualAnimSelectionCriterion.cpp
+++ b/Engine/Plugins/Experimental/Animation/ContextualAnimation/Source/ContextualAnimation/Private/ContextualAnimSelectionCriterion.cpp
@@ -80,53 +80,39 @@ bool UContextualAnimSelectionCriterion_TriggerArea::DoesQuerierPassCondition(con
 	return bResult;
 }
 
-// UContextualAnimSelectionCriterion_Angle
+// UContextualAnimSelectionCriterion_Cone
 //===========================================================================
 
-bool UContextualAnimSelectionCriterion_Angle::DoesQuerierPassCondition(const FContextualAnimSceneBindingContext& Primary, const FContextualAnimSceneBindingContext& Querier) const
+bool UContextualAnimSelectionCriterion_Cone::DoesQuerierPassCondition(const FContextualAnimSceneBindingContext& Primary, const FContextualAnimSceneBindingContext& Querier) const
 {
+	bool bResult = false;
+
 	const FTransform PrimaryTransform = Primary.GetTransform();
 	const FTransform QuerierTransform = Querier.GetTransform();
+	const float ActualDist = FVector::Dist(PrimaryTransform.GetLocation(), QuerierTransform.GetLocation());
 
-	auto CalculateAngle = [](const FTransform& A, const FTransform& B, bool bSignedAngle)
+	if(ActualDist <= Distance)
 	{
-		const FVector ToTarget = (A.GetLocation() - B.GetLocation()).GetSafeNormal2D();
-		const float ForwardCosAngle = FVector::DotProduct(B.GetRotation().GetForwardVector(), ToTarget);
-		const float ForwardDeltaDegree = FMath::RadiansToDegrees(FMath::Acos(ForwardCosAngle));
+		auto CheckCondition = [this](const FTransform& A, const FTransform& B)
+		{
+			const FVector VecA = (A.GetLocation() - B.GetLocation()).GetSafeNormal2D();
+			const FVector VecB = B.GetRotation().GetForwardVector().RotateAngleAxis(Offset, FVector::UpVector);
+			return FVector::DotProduct(VecA, VecB) > FMath::Cos(FMath::Clamp(FMath::DegreesToRadians(HalfAngle), 0.f, PI));
+		};
 		
-		if(bSignedAngle)
+		if (Mode == EContextualAnimCriterionConeMode::ToPrimary)
 		{
-			const float RightCosAngle = FVector::DotProduct(B.GetRotation().GetRightVector(), ToTarget);
-			return (RightCosAngle < 0) ? (ForwardDeltaDegree * -1) : ForwardDeltaDegree;
+			bResult = CheckCondition(PrimaryTransform, QuerierTransform);
 		}
-		else
+		else if (Mode == EContextualAnimCriterionConeMode::FromPrimary)
 		{
-			return ForwardDeltaDegree;
+			bResult = CheckCondition(QuerierTransform, PrimaryTransform);
 		}
-	};
-
-	float Angle = 0.f;
-	if (Mode == EContextualAnimCriterionAngleMode::ToPrimary)
-	{
-		Angle = CalculateAngle(PrimaryTransform, QuerierTransform, bUseSignedAngle);
-	}
-	else if (Mode == EContextualAnimCriterionAngleMode::FromPrimary)
-	{
-		Angle = CalculateAngle(QuerierTransform, PrimaryTransform, bUseSignedAngle);
-	}
-
-	bool bResult = false;
-	if (MinAngle <= 0 && MaxAngle >= 0)
-	{
-		bResult = FMath::IsWithinInclusive(FMath::Abs(Angle), FMath::Abs(MinAngle), FMath::Abs(MaxAngle));
-	}
-	else
-	{
-		bResult = FMath::IsWithinInclusive(Angle, MinAngle, MaxAngle);
 	}
+	
 
-	UE_LOG(LogContextualAnim, Verbose, TEXT("UContextualAnimSelectionCriterion_Angle: Primary: %s Querier: %s Mode: %s bUseSignedAngle: %d MinAngle: %.1f MaxAngle: %.1f Angle: %.1f Result: %d"), 
-		*GetNameSafe(Primary.GetActor()), *GetNameSafe(Querier.GetActor()), *UEnum::GetValueAsString(TEXT("ContextualAnimation.EContextualAnimCriterionAngleMode"), Mode), bUseSignedAngle, MinAngle, MaxAngle, Angle, bResult);
+	UE_LOG(LogContextualAnim, Verbose, TEXT("UContextualAnimSelectionCriterion_Cone: Primary: %s Querier: %s Mode: %s Distance: %.1f HalfAngle: %.1f Offset: %.1f Result: %d"), 
+		*GetNameSafe(Primary.GetActor()), *GetNameSafe(Querier.GetActor()), *UEnum::GetValueAsString(TEXT("ContextualAnimation.EContextualAnimCriterionAngleMode"), Mode), Distance, HalfAngle, Offset, bResult);
 
 	return bResult;
 }
diff --git a/Engine/Plugins/Experimental/Animation/ContextualAnimation/Source/ContextualAnimation/Public/ContextualAnimSelectionCriterion.h b/Engine/Plugins/Experimental/Animation/ContextualAnimation/Source/ContextualAnimation/Public/ContextualAnimSelectionCriterion.h
--- a/Engine/Plugins/Experimental/Animation/ContextualAnimation/Source/ContextualAnimation/Public/ContextualAnimSelectionCriterion.h
+++ b/Engine/Plugins/Experimental/Animation/ContextualAnimation/Source/ContextualAnimation/Public/ContextualAnimSelectionCriterion.h
@@ -76,39 +76,39 @@ public:
 	virtual bool DoesQuerierPassCondition(const FContextualAnimSceneBindingContext& Primary, const FContextualAnimSceneBindingContext& Querier) const override;
 };
 
-// UContextualAnimSelectionCriterion_Angle
+// UContextualAnimSelectionCriterion_Cone
 //===========================================================================
 
 UENUM(BlueprintType)
-enum class EContextualAnimCriterionAngleMode : uint8
+enum class EContextualAnimCriterionConeMode : uint8
 {
-	/** Uses the angle between the vector from querier to primary and querier forward vector */
+	/** Uses the angle between the vector from querier to primary and querier forward vector rotated by offset */
 	ToPrimary,
 
-	/** Uses the angle between the vector from primary to querier and primary forward vector */
+	/** Uses the angle between the vector from primary to querier and primary forward vector rotated by offset */
 	FromPrimary
 };
 
 UCLASS()
-class CONTEXTUALANIMATION_API UContextualAnimSelectionCriterion_Angle : public UContextualAnimSelectionCriterion
+class CONTEXTUALANIMATION_API UContextualAnimSelectionCriterion_Cone : public UContextualAnimSelectionCriterion
 {
 	GENERATED_BODY()
 
 public:
 
 	UPROPERTY(EditAnywhere, BlueprintReadWrite, Category = "Default")
-	EContextualAnimCriterionAngleMode Mode = EContextualAnimCriterionAngleMode::ToPrimary;
+	EContextualAnimCriterionConeMode Mode = EContextualAnimCriterionConeMode::ToPrimary;
 
-	UPROPERTY(EditAnywhere, BlueprintReadWrite, Category = "Default")
-	bool bUseSignedAngle = false;
+	UPROPERTY(EditAnywhere, BlueprintReadWrite, Category = "Default", meta = (ClampMin = "0", UIMin = "0"))
+	float Distance = 200.f;
 
-	UPROPERTY(EditAnywhere, BlueprintReadWrite, Category = "Default", meta = (ClampMin = "-180", UIMin = "-180", ClampMax = "180", UIMax = "180"))
-	float MinAngle = 0.f;
+	UPROPERTY(EditAnywhere, BlueprintReadWrite, Category = "Default", meta = (ClampMin = "0", UIMin = "0", ClampMax = "180", UIMax = "180"))
+	float HalfAngle = 45.f;
 
 	UPROPERTY(EditAnywhere, BlueprintReadWrite, Category = "Default", meta = (ClampMin = "-180", UIMin = "-180", ClampMax = "180", UIMax = "180"))
-	float MaxAngle = 0.f;
+	float Offset = 0.f;
 
-	UContextualAnimSelectionCriterion_Angle(const FObjectInitializer& ObjectInitializer) : Super(ObjectInitializer) {}
+	UContextualAnimSelectionCriterion_Cone(const FObjectInitializer& ObjectInitializer) : Super(ObjectInitializer) {}
 
 	virtual bool DoesQuerierPassCondition(const FContextualAnimSceneBindingContext& Primary, const FContextualAnimSceneBindingContext& Querier) const override;
 };
diff --git a/Engine/Plugins/Experimental/Animation/ContextualAnimation/Source/ContextualAnimationEditor/Private/ContextualAnimEdMode.cpp b/Engine/Plugins/Experimental/Animation/ContextualAnimation/Source/ContextualAnimationEditor/Private/ContextualAnimEdMode.cpp
--- a/Engine/Plugins/Experimental/Animation/ContextualAnimation/Source/ContextualAnimationEditor/Private/ContextualAnimEdMode.cpp
+++ b/Engine/Plugins/Experimental/Animation/ContextualAnimation/Source/ContextualAnimationEditor/Private/ContextualAnimEdMode.cpp
@@ -114,6 +114,24 @@ void FContextualAnimEdMode::Render(const FSceneView* View, FViewport* Viewport,
 									PDI->SetHitProxy(nullptr);
 								}
 							}
+							else if (const UContextualAnimSelectionCriterion_Cone* Cone = Cast<UContextualAnimSelectionCriterion_Cone>(Criterion))
+							{
+								const FTransform QuerierTransform = Binding.GetContext().GetTransform();
+
+								FVector Origin, Direction;
+								if (Cone->Mode == EContextualAnimCriterionConeMode::ToPrimary)
+								{
+									Origin = QuerierTransform.GetLocation();
+									Direction = QuerierTransform.GetRotation().GetForwardVector().RotateAngleAxis(Cone->Offset, FVector::UpVector);
+								}
+								else if (Cone->Mode == EContextualAnimCriterionConeMode::FromPrimary)
+								{
+									Origin = PrimaryTransform.GetLocation();
+									Direction = PrimaryTransform.GetRotation().GetForwardVector().RotateAngleAxis(Cone->Offset, FVector::UpVector);
+								}
+
+								UContextualAnimUtilities::DrawSector(*PDI, Origin, Direction, 0.f, Cone->Distance, -Cone->HalfAngle, Cone->HalfAngle, DrawColor, SDPG_World, 1.f);
+							}
 						}
 					}
 				}
