ContextualAnimEditor: Fixed bug where selection widget and draw debugs randomly disappear from the preview scene

diff --git a/Engine/Plugins/Experimental/Animation/ContextualAnimation/Source/ContextualAnimationEditor/Private/ContextualAnimEdMode.cpp b/Engine/Plugins/Experimental/Animation/ContextualAnimation/Source/ContextualAnimationEditor/Private/ContextualAnimEdMode.cpp
--- a/Engine/Plugins/Experimental/Animation/ContextualAnimation/Source/ContextualAnimationEditor/Private/ContextualAnimEdMode.cpp
+++ b/Engine/Plugins/Experimental/Animation/ContextualAnimation/Source/ContextualAnimationEditor/Private/ContextualAnimEdMode.cpp
@@ -5,7 +5,6 @@
 #include "ContextualAnimSceneAsset.h"
 #include "ContextualAnimSelectionCriterion.h"
 #include "ContextualAnimUtilities.h"
-#include "ContextualAnimSceneInstance.h"
 #include "ContextualAnimViewportClient.h"
 #include "ContextualAnimAssetEditorToolkit.h"
 #include "ContextualAnimViewModel.h"
@@ -39,22 +38,21 @@ void FContextualAnimEdMode::Render(const FSceneView* View, FViewport* Viewport,
 
 	// @TODO: This should not be initialized here
 	FContextualAnimViewportClient* ViewportClient = static_cast<FContextualAnimViewportClient*>(Viewport->GetClient());
-	if(!ViewModel)
+	if (!ViewModel)
 	{
 		ViewModel = ViewportClient->GetAssetEditorToolkit()->GetViewModel();
 	}
 
 	if (ViewModel)
 	{
-		if (const UContextualAnimSceneInstance* SceneInstance = ViewModel->GetSceneInstance())
+		if (const UContextualAnimSceneAsset* SceneAsset = ViewModel->GetSceneAsset())
 		{
-			const UContextualAnimSceneAsset& SceneAsset = SceneInstance->GetSceneAsset();
-			const FContextualAnimSceneBindings& Bindings = SceneInstance->GetBindings();
+			const FContextualAnimSceneBindings& Bindings = ViewModel->SceneBindings;
 
 			// Draw Scene Pivots
-			if(const FContextualAnimSceneSection* Section = SceneAsset.GetSection(Bindings.GetSectionIdx()))
+			if (const FContextualAnimSceneSection* Section = SceneAsset->GetSection(Bindings.GetSectionIdx()))
 			{
-				if(const FContextualAnimSet* AnimSet = Section->GetAnimSet(Bindings.GetAnimSetIdx()))
+				if (const FContextualAnimSet* AnimSet = Section->GetAnimSet(Bindings.GetAnimSetIdx()))
 				{
 					for (const FTransform& ScenePivot : AnimSet->ScenePivots)
 					{
@@ -74,14 +72,14 @@ void FContextualAnimEdMode::Render(const FSceneView* View, FViewport* Viewport,
 				}
 
 				// Draw Selection Criteria
-				if(const FContextualAnimSceneBinding* PrimaryBinding = Bindings.FindBindingByRole(SceneAsset.GetPrimaryRole()))
+				if (const FContextualAnimSceneBinding* PrimaryBinding = Bindings.FindBindingByRole(SceneAsset->GetPrimaryRole()))
 				{
 					const FTransform PrimaryTransform = PrimaryBinding->GetTransform();
 
 					const FContextualAnimTrack& AnimTrack = Binding.GetAnimTrack();
 					for (int32 CriterionIdx = 0; CriterionIdx < AnimTrack.SelectionCriteria.Num(); CriterionIdx++)
 					{
-						if(const UContextualAnimSelectionCriterion* Criterion = AnimTrack.SelectionCriteria[CriterionIdx])
+						if (const UContextualAnimSelectionCriterion* Criterion = AnimTrack.SelectionCriteria[CriterionIdx])
 						{
 							FLinearColor DrawColor = FLinearColor::White;
 							if (Criterion->DoesQuerierPassCondition(FContextualAnimSceneBindingContext(PrimaryTransform), Binding.GetContext()))
@@ -155,7 +153,7 @@ void FContextualAnimEdMode::DrawIKTargetsForBinding(FPrimitiveDrawInterface& PDI
 {
 	for (const FContextualAnimIKTargetDefinition& IKTargetDef : Binding.GetIKTargetDefs().IKTargetDefs)
 	{
-		if (const FContextualAnimSceneBinding* TargetBinding = Binding.GetSceneInstance()->FindBindingByRole(IKTargetDef.TargetRoleName))
+		if (const FContextualAnimSceneBinding* TargetBinding = ViewModel->SceneBindings.FindBindingByRole(IKTargetDef.TargetRoleName))
 		{
 			if (USkeletalMeshComponent* TargetSkelMeshComp = TargetBinding->GetSkeletalMeshComponent())
 			{
@@ -257,7 +255,7 @@ bool FContextualAnimEdMode::InputDelta(FEditorViewportClient* InViewportClient,
 
 bool FContextualAnimEdMode::InputKey(FEditorViewportClient* ViewportClient, FViewport* Viewport, FKey Key, EInputEvent Event)
 {
-	if(ViewModel && ViewModel->IsSimulateModePaused())
+	if (ViewModel && ViewModel->IsSimulateModePaused())
 	{
 		if (Key == EKeys::Enter && Event == IE_Released)
 		{
diff --git a/Engine/Plugins/Experimental/Animation/ContextualAnimation/Source/ContextualAnimationEditor/Private/ContextualAnimViewModel.cpp b/Engine/Plugins/Experimental/Animation/ContextualAnimation/Source/ContextualAnimationEditor/Private/ContextualAnimViewModel.cpp
--- a/Engine/Plugins/Experimental/Animation/ContextualAnimation/Source/ContextualAnimationEditor/Private/ContextualAnimViewModel.cpp
+++ b/Engine/Plugins/Experimental/Animation/ContextualAnimation/Source/ContextualAnimationEditor/Private/ContextualAnimViewModel.cpp
@@ -352,19 +352,14 @@ void FContextualAnimViewModel::RefreshPreviewScene()
 		SceneInstance.Reset();
 	}
 
-	if (StartSceneParams.RoleToActorMap.Num() > 0)
+	for (const FContextualAnimSceneBinding& Binding : SceneBindings)
 	{
-		for (auto& MapEntry : StartSceneParams.RoleToActorMap)
+		if (AActor* Actor = Binding.GetActor())
 		{
-			if (AActor* Actor = MapEntry.Value.GetActor())
-			{
-				Actor->Destroy();
-			}
+			Actor->Destroy();
 		}
 	}
 
-	StartSceneParams.Reset();
-
 	if (!SceneAsset->Sections.IsValidIndex(ActiveSectionIdx))
 	{
 		return;
@@ -377,6 +372,10 @@ void FContextualAnimViewModel::RefreshPreviewScene()
 		return;
 	}
 	
+	FContextualAnimStartSceneParams StartSceneParams;
+	StartSceneParams.SectionIdx = ActiveSectionIdx;
+	StartSceneParams.AnimSetIdx = ActiveAnimSetIdx;
+
 	const FContextualAnimSet& AnimSet = ContextualAnimSection.AnimSets[ActiveAnimSetIdx];
 	for (const FContextualAnimTrack& AnimTrack : AnimSet.Tracks)
 	{
@@ -386,15 +385,14 @@ void FContextualAnimViewModel::RefreshPreviewScene()
 		}
 	}
 
-	StartSceneParams.SectionIdx = ActiveSectionIdx;
-	StartSceneParams.AnimSetIdx = ActiveAnimSetIdx;
-
 	SceneInstance = ContextualAnimManager->ForceStartScene(*GetSceneAsset(), StartSceneParams);
 
 	// Disable auto blend out
 	if (SceneInstance.IsValid())
 	{
-		for (auto& Binding : SceneInstance->GetBindings())
+		SceneBindings = SceneInstance->GetBindings();
+
+		for (const FContextualAnimSceneBinding& Binding : SceneBindings)
 		{
 			if (FAnimMontageInstance* MontageInstance = Binding.GetAnimMontageInstance())
 			{
@@ -515,7 +513,10 @@ void FContextualAnimViewModel::ToggleSimulateMode()
 void FContextualAnimViewModel::StartSimulation()
 {
 	FContextualAnimStartSceneParams Params;
-	Params.RoleToActorMap = StartSceneParams.RoleToActorMap;
+	for (const FContextualAnimSceneBinding& Binding : SceneBindings)
+	{
+		Params.RoleToActorMap.Add(Binding.GetRoleDef().Name, Binding.GetContext());
+	}
 	Params.SectionIdx = 0;
 	SceneInstance = ContextualAnimManager->TryStartScene(*GetSceneAsset(), Params);
 
@@ -526,6 +527,8 @@ void FContextualAnimViewModel::StartSimulation()
 		return;
 	}
 
+	SceneBindings = SceneInstance->GetBindings();
+
 	SimulateModeState = ESimulateModeState::Playing;
 }
 
@@ -717,7 +720,7 @@ void FContextualAnimViewModel::AnimationModified(UAnimSequenceBase& Animation)
 
 void FContextualAnimViewModel::UpdateSelection(const AActor* SelectedActor)
 {
-	const FContextualAnimSceneBinding* Binding = SceneInstance.IsValid() ? SceneInstance->FindBindingByActor(SelectedActor) : nullptr;
+	const FContextualAnimSceneBinding* Binding = SceneBindings.FindBindingByActor(SelectedActor);
 	if (Binding)
 	{
 		// Discard changes if we were modifying the transform of an actor in the scene but never confirmed the changes before selecting a different actor
@@ -751,7 +754,7 @@ void FContextualAnimViewModel::ClearSelection()
 
 FContextualAnimSceneBinding* FContextualAnimViewModel::GetSelectedBinding() const
 {
-	return SceneInstance.IsValid() ? const_cast<FContextualAnimSceneBinding*>(SceneInstance->FindBindingByRole(SelectionInfo.Role)) : nullptr;
+	return const_cast<FContextualAnimSceneBinding*>(SceneBindings.FindBindingByRole(SelectionInfo.Role));
 }
 
 AActor* FContextualAnimViewModel::GetSelectedActor() const
diff --git a/Engine/Plugins/Experimental/Animation/ContextualAnimation/Source/ContextualAnimationEditor/Public/ContextualAnimViewModel.h b/Engine/Plugins/Experimental/Animation/ContextualAnimation/Source/ContextualAnimationEditor/Public/ContextualAnimViewModel.h
--- a/Engine/Plugins/Experimental/Animation/ContextualAnimation/Source/ContextualAnimationEditor/Public/ContextualAnimViewModel.h
+++ b/Engine/Plugins/Experimental/Animation/ContextualAnimation/Source/ContextualAnimationEditor/Public/ContextualAnimViewModel.h
@@ -34,6 +34,8 @@ class FContextualAnimViewModel : public TSharedFromThis<FContextualAnimViewModel
 {
 public:
 
+	friend class FContextualAnimEdMode;
+
 	enum class ETimelineMode : uint8
 	{
 		Default,
@@ -163,6 +165,9 @@ private:
 
 	TWeakObjectPtr<UContextualAnimSceneInstance> SceneInstance;
 
+	/** Copy of the bindings so we can access them even when simulation mode is not playing (and SceneInstance is not valid) */
+	FContextualAnimSceneBindings SceneBindings;
+
 	/** Active section idx */
 	int32 ActiveSectionIdx = 0;
 
@@ -180,8 +185,6 @@ private:
 
 	ESimulateModeState SimulateModeState = ESimulateModeState::Inactive;
 
-	FContextualAnimStartSceneParams StartSceneParams;
-
 	ETimelineMode TimelineMode = ETimelineMode::Default;
 
 	TWeakObjectPtr<UAnimSequenceBase> EditingAnimation = nullptr;
