Contextual Anim Editor: Refactored visualization of selection criteria and added options to the viewport to decide when to draw them + added similar options to draw entry poses

diff --git a/Engine/Plugins/Experimental/Animation/ContextualAnimation/Source/ContextualAnimation/Private/ContextualAnimUtilities.cpp b/Engine/Plugins/Experimental/Animation/ContextualAnimation/Source/ContextualAnimation/Private/ContextualAnimUtilities.cpp
--- a/Engine/Plugins/Experimental/Animation/ContextualAnimation/Source/ContextualAnimation/Private/ContextualAnimUtilities.cpp
+++ b/Engine/Plugins/Experimental/Animation/ContextualAnimation/Source/ContextualAnimation/Private/ContextualAnimUtilities.cpp
@@ -86,48 +86,79 @@ FTransform UContextualAnimUtilities::ExtractRootTransformFromAnimation(const UAn
 	return FTransform::Identity;
 }
 
-void UContextualAnimUtilities::DrawDebugPose(const UWorld* World, const UAnimSequenceBase* Animation, float Time, const FTransform& LocalToWorldTransform, const FColor& Color, float LifeTime, float Thickness)
-{ 
-	if(World)
+void UContextualAnimUtilities::BP_DrawDebugPose(const UObject* WorldContextObject, const UAnimSequenceBase* Animation, float Time, FTransform LocalToWorldTransform, FLinearColor Color, float LifeTime, float Thickness)
+{
+	if(GEngine)
 	{
-		FMemMark Mark(FMemStack::Get());
-
-		Time = FMath::Clamp(Time, 0.f, Animation->GetPlayLength());
-
-		const int32 TotalBones = Animation->GetSkeleton()->GetReferenceSkeleton().GetNum();
-		TArray<FBoneIndexType> RequiredBoneIndexArray;
-		RequiredBoneIndexArray.Reserve(TotalBones);
-		for (int32 Idx = 0; Idx < TotalBones; Idx++)
+		if (UWorld* World = GEngine->GetWorldFromContextObject(WorldContextObject, EGetWorldErrorMode::LogAndReturnNull))
 		{
-			RequiredBoneIndexArray.Add(Idx);
+			DrawPose(World, Animation, Time, LocalToWorldTransform, Color, LifeTime, Thickness);
 		}
+	}
+}
 
-		FBoneContainer BoneContainer(RequiredBoneIndexArray, FCurveEvaluationOption(false), *Animation->GetSkeleton());
-		FCSPose<FCompactPose> ComponentSpacePose;
-		UContextualAnimUtilities::ExtractComponentSpacePose(Animation, BoneContainer, Time, true, ComponentSpacePose);
+void UContextualAnimUtilities::DrawPose(const UWorld* World, const UAnimSequenceBase* Animation, float Time, FTransform LocalToWorldTransform, FLinearColor Color, float LifeTime, float Thickness)
+{
+	if (World)
+	{
+		auto DrawFunction = [World](const FVector& LineStart, const FVector& LineEnd, const FColor& Color, float LifeTime, float Thickness) {
+			DrawDebugLine(World, LineStart, LineEnd, Color, false, LifeTime, 0, Thickness);
+		};
 
-		for (int32 Index = 0; Index < ComponentSpacePose.GetPose().GetNumBones(); ++Index)
-		{
-			const FCompactPoseBoneIndex CompactPoseBoneIndex = FCompactPoseBoneIndex(Index);
-			const FCompactPoseBoneIndex ParentIndex = ComponentSpacePose.GetPose().GetParentBoneIndex(CompactPoseBoneIndex);
-			FVector Start, End;
+		DrawPose(Animation, Time, LocalToWorldTransform, Color, LifeTime, Thickness, DrawFunction);
+	}
+}
+
+void UContextualAnimUtilities::DrawPose(FPrimitiveDrawInterface* PDI, const UAnimSequenceBase* Animation, float Time, FTransform LocalToWorldTransform, FLinearColor Color, float Thickness)
+{
+	if (PDI)
+	{
+		auto DrawFunction = [PDI](const FVector& LineStart, const FVector& LineEnd, const FColor& Color, float LifeTime, float Thickness) {
+			PDI->DrawLine(LineStart, LineEnd, Color, 0, Thickness);
+		};
 
-			FLinearColor LineColor = FLinearColor::Red;
-			const FTransform Transform = ComponentSpacePose.GetComponentSpaceTransform(CompactPoseBoneIndex) * LocalToWorldTransform;
+		DrawPose(Animation, Time, LocalToWorldTransform, Color, 0, Thickness, DrawFunction);
+	}
+}
 
-			if (ParentIndex.GetInt() >= 0)
-			{
-				Start = (ComponentSpacePose.GetComponentSpaceTransform(ParentIndex) * LocalToWorldTransform).GetLocation();
-				End = Transform.GetLocation();
-			}
-			else
-			{
-				Start = LocalToWorldTransform.GetLocation();
-				End = Transform.GetLocation();
-			}
+void UContextualAnimUtilities::DrawPose(const UAnimSequenceBase* Animation, float Time, FTransform LocalToWorldTransform, FLinearColor Color, float LifeTime, float Thickness, FDrawLineFunction DrawFunction)
+{
+	FMemMark Mark(FMemStack::Get());
+
+	Time = FMath::Clamp(Time, 0.f, Animation->GetPlayLength());
+
+	const int32 TotalBones = Animation->GetSkeleton()->GetReferenceSkeleton().GetNum();
+	TArray<FBoneIndexType> RequiredBoneIndexArray;
+	RequiredBoneIndexArray.Reserve(TotalBones);
+	for (int32 Idx = 0; Idx < TotalBones; Idx++)
+	{
+		RequiredBoneIndexArray.Add(Idx);
+	}
 
-			DrawDebugLine(World, Start, End, Color, false, LifeTime, 0, Thickness);
+	FBoneContainer BoneContainer(RequiredBoneIndexArray, FCurveEvaluationOption(false), *Animation->GetSkeleton());
+	FCSPose<FCompactPose> ComponentSpacePose;
+	UContextualAnimUtilities::ExtractComponentSpacePose(Animation, BoneContainer, Time, true, ComponentSpacePose);
+
+	for (int32 Index = 0; Index < ComponentSpacePose.GetPose().GetNumBones(); ++Index)
+	{
+		const FCompactPoseBoneIndex CompactPoseBoneIndex = FCompactPoseBoneIndex(Index);
+		const FCompactPoseBoneIndex ParentIndex = ComponentSpacePose.GetPose().GetParentBoneIndex(CompactPoseBoneIndex);
+		FVector Start, End;
+
+		const FTransform Transform = ComponentSpacePose.GetComponentSpaceTransform(CompactPoseBoneIndex) * LocalToWorldTransform;
+
+		if (ParentIndex.GetInt() >= 0)
+		{
+			Start = (ComponentSpacePose.GetComponentSpaceTransform(ParentIndex) * LocalToWorldTransform).GetLocation();
+			End = Transform.GetLocation();
 		}
+		else
+		{
+			Start = LocalToWorldTransform.GetLocation();
+			End = Transform.GetLocation();
+		}
+
+		DrawFunction(Start, End, Color.ToFColor(false), LifeTime, Thickness);
 	}
 }
 
@@ -141,7 +172,7 @@ void UContextualAnimUtilities::DrawDebugAnimSet(const UWorld* World, const UCont
 
 			if (const UAnimSequenceBase* Animation = AnimTrack.Animation)
 			{
-				DrawDebugPose(World, Animation, Time, Transform, Color, LifeTime, Thickness);
+				DrawPose(World, Animation, Time, Transform, Color, LifeTime, Thickness);
 			}
 			else
 			{
@@ -212,7 +243,7 @@ float UContextualAnimUtilities::BP_Montage_GetSectionLength(const UAnimMontage*
 	return Montage ? Montage->GetSectionLength(SectionIndex) : -1.f;
 }
 
-void UContextualAnimUtilities::DrawSector(FPrimitiveDrawInterface& PDI, const FVector& Origin, const FVector& Direction, float MinDistance, float MaxDistance, float MinAngle, float MaxAngle, const FLinearColor& Color, uint8 DepthPriority, float Thickness)
+void UContextualAnimUtilities::DrawSector(FPrimitiveDrawInterface& PDI, const FVector& Origin, const FVector& Direction, float MinDistance, float MaxDistance, float MinAngle, float MaxAngle, const FLinearColor& Color, uint8 DepthPriority, float Thickness, bool bDashedLine)
 {
 	if(MinAngle == 0 && MaxAngle == 0)
 	{
@@ -223,8 +254,17 @@ void UContextualAnimUtilities::DrawSector(FPrimitiveDrawInterface& PDI, const FV
 	// Draw Cone lines
 	const FVector LeftDirection = Direction.RotateAngleAxis(MinAngle, FVector::UpVector);
 	const FVector RightDirection = Direction.RotateAngleAxis(MaxAngle, FVector::UpVector);
-	PDI.DrawLine(Origin + (LeftDirection * MinDistance), Origin + (LeftDirection * MaxDistance), Color, DepthPriority, Thickness);
-	PDI.DrawLine(Origin + (RightDirection * MinDistance), Origin + (RightDirection * MaxDistance), Color, DepthPriority, Thickness);
+
+	if(bDashedLine)
+	{
+		DrawDashedLine(&PDI, Origin + (LeftDirection * MinDistance), Origin + (LeftDirection * MaxDistance), Color, 10.f, DepthPriority);
+		DrawDashedLine(&PDI, Origin + (RightDirection * MinDistance), Origin + (RightDirection * MaxDistance), Color, 10.f, DepthPriority);
+	}
+	else
+	{
+		PDI.DrawLine(Origin + (LeftDirection * MinDistance), Origin + (LeftDirection * MaxDistance), Color, DepthPriority, Thickness);
+		PDI.DrawLine(Origin + (RightDirection * MinDistance), Origin + (RightDirection * MaxDistance), Color, DepthPriority, Thickness);
+	}
 
 	// Draw Near Arc
 	FVector LastDirection = LeftDirection;
@@ -237,7 +277,16 @@ void UContextualAnimUtilities::DrawSector(FPrimitiveDrawInterface& PDI, const FV
 		const FVector NewDirection = Direction.RotateAngleAxis(Angle, FVector::UpVector);
 		const FVector LineStart = Origin + (LastDirection * Length);
 		const FVector LineEnd = Origin + (NewDirection * Length);
-		PDI.DrawLine(LineStart, LineEnd, Color, DepthPriority, Thickness);
+		
+		if (bDashedLine)
+		{
+			DrawDashedLine(&PDI, LineStart, LineEnd, Color, 10.f, DepthPriority);
+		}
+		else
+		{
+			PDI.DrawLine(LineStart, LineEnd, Color, DepthPriority, Thickness);
+		}
+		
 		LastDirection = NewDirection;
 	}
 
@@ -252,7 +301,16 @@ void UContextualAnimUtilities::DrawSector(FPrimitiveDrawInterface& PDI, const FV
 		const FVector NewDirection = Direction.RotateAngleAxis(Angle, FVector::UpVector);
 		const FVector LineStart = Origin + (LastDirection * Length);
 		const FVector LineEnd = Origin + (NewDirection * Length);
-		PDI.DrawLine(LineStart, LineEnd, Color, DepthPriority, Thickness);
+
+		if (bDashedLine)
+		{
+			DrawDashedLine(&PDI, LineStart, LineEnd, Color, 10.f, DepthPriority);
+		}
+		else
+		{
+			PDI.DrawLine(LineStart, LineEnd, Color, DepthPriority, Thickness);
+		}
+
 		LastDirection = NewDirection;
 	}
 }
diff --git a/Engine/Plugins/Experimental/Animation/ContextualAnimation/Source/ContextualAnimation/Public/ContextualAnimSceneAsset.h b/Engine/Plugins/Experimental/Animation/ContextualAnimation/Source/ContextualAnimation/Public/ContextualAnimSceneAsset.h
--- a/Engine/Plugins/Experimental/Animation/ContextualAnimation/Source/ContextualAnimation/Public/ContextualAnimSceneAsset.h
+++ b/Engine/Plugins/Experimental/Animation/ContextualAnimation/Source/ContextualAnimation/Public/ContextualAnimSceneAsset.h
@@ -72,6 +72,7 @@ public:
 
 	FORCEINLINE FName GetName() const { return Name; }
 	FORCEINLINE const TArray<FContextualAnimSetPivotDefinition>& GetAnimSetPivotDefinitions() const { return AnimSetPivotDefinitions; }
+	FORCEINLINE int32 GetNumAnimSets() const { return AnimSets.Num(); }
 
 protected:
 
diff --git a/Engine/Plugins/Experimental/Animation/ContextualAnimation/Source/ContextualAnimation/Public/ContextualAnimUtilities.h b/Engine/Plugins/Experimental/Animation/ContextualAnimation/Source/ContextualAnimation/Public/ContextualAnimUtilities.h
--- a/Engine/Plugins/Experimental/Animation/ContextualAnimation/Source/ContextualAnimation/Public/ContextualAnimUtilities.h
+++ b/Engine/Plugins/Experimental/Animation/ContextualAnimation/Source/ContextualAnimation/Public/ContextualAnimUtilities.h
@@ -42,7 +42,13 @@ public:
 	/** Extract root bone transform at a given time */
 	static FTransform ExtractRootTransformFromAnimation(const UAnimSequenceBase* Animation, float Time);
 
-	static void DrawDebugPose(const UWorld* World, const UAnimSequenceBase* Animation, float Time, const FTransform& LocalToWorldTransform, const FColor& Color, float LifeTime, float Thickness);
+	UFUNCTION(BlueprintCallable, Category = "Contextual Anim|Utilities", meta = (WorldContext = "WorldContextObject", DisplayName = "Draw Debug Pose"))
+	static void BP_DrawDebugPose(const UObject* WorldContextObject, const UAnimSequenceBase* Animation, float Time, FTransform LocalToWorldTransform, FLinearColor Color, float LifeTime, float Thickness);
+
+	typedef TFunctionRef<void(const FVector& LineStart, const FVector& LineEnd, const FColor& Color, float LifeTime, float Thickness)> FDrawLineFunction;
+	static void DrawPose(const UWorld* World, const UAnimSequenceBase* Animation, float Time, FTransform LocalToWorldTransform, FLinearColor Color, float LifeTime, float Thickness);
+	static void DrawPose(FPrimitiveDrawInterface* PDI, const UAnimSequenceBase* Animation, float Time, FTransform LocalToWorldTransform, FLinearColor Color, float Thickness);
+	static void DrawPose(const UAnimSequenceBase* Animation, float Time, FTransform LocalToWorldTransform, FLinearColor Color, float LifeTime, float Thickness, FDrawLineFunction DrawFunction);
 	
 	static void DrawDebugAnimSet(const UWorld* World, const UContextualAnimSceneAsset& SceneAsset, const FContextualAnimSet& AnimSet, float Time, const FTransform& ToWorldTransform, const FColor& Color, float LifeTime, float Thickness);
 
@@ -52,7 +58,7 @@ public:
 
 	static FAnimMontageInstance* TryGetActiveAnimMontageInstance(const AActor* Actor);
 
-	static void DrawSector(FPrimitiveDrawInterface& PDI, const FVector& Origin, const FVector& Direction, float MinDistance, float MaxDistance, float MinAngle, float MaxAngle, const FLinearColor& Color, uint8 DepthPriority, float Thickness);
+	static void DrawSector(FPrimitiveDrawInterface& PDI, const FVector& Origin, const FVector& Direction, float MinDistance, float MaxDistance, float MinAngle, float MaxAngle, const FLinearColor& Color, uint8 DepthPriority, float Thickness, bool bDashedLine);
 
 	UFUNCTION(BlueprintCallable, Category = "Contextual Anim|Scene Asset", meta = (DisplayName = "Create Contextual Anim Scene Bindings"))
 	static bool BP_CreateContextualAnimSceneBindings(const UContextualAnimSceneAsset* SceneAsset, const TMap<FName, FContextualAnimSceneBindingContext>& Params, FContextualAnimSceneBindings& OutBindings);
diff --git a/Engine/Plugins/Experimental/Animation/ContextualAnimation/Source/ContextualAnimationEditor/Private/ContextualAnimAssetEditorCommands.cpp b/Engine/Plugins/Experimental/Animation/ContextualAnimation/Source/ContextualAnimationEditor/Private/ContextualAnimAssetEditorCommands.cpp
--- a/Engine/Plugins/Experimental/Animation/ContextualAnimation/Source/ContextualAnimationEditor/Private/ContextualAnimAssetEditorCommands.cpp
+++ b/Engine/Plugins/Experimental/Animation/ContextualAnimation/Source/ContextualAnimationEditor/Private/ContextualAnimAssetEditorCommands.cpp
@@ -13,6 +13,14 @@ void FContextualAnimAssetEditorCommands::RegisterCommands()
 	UI_COMMAND(ShowIKTargetsDrawAll, "All Actors", "Show IK Targets for all the actors", EUserInterfaceActionType::RadioButton, FInputChord());
 	UI_COMMAND(ShowIKTargetsDrawNone, "None", "Hide IK Targets", EUserInterfaceActionType::RadioButton, FInputChord());
 
+	UI_COMMAND(ShowSelectionCriteriaActiveSet, "Active AnimSet", "Show selection criteria only for the active AnimSet", EUserInterfaceActionType::RadioButton, FInputChord());
+	UI_COMMAND(ShowSelectionCriteriaAllSets, "All AnimSets", "Show selection criteria for all the AnimSets", EUserInterfaceActionType::RadioButton, FInputChord());
+	UI_COMMAND(ShowSelectionCriteriaNone, "None", "Hide selection criteria", EUserInterfaceActionType::RadioButton, FInputChord());
+
+	UI_COMMAND(ShowEntryPosesActiveSet, "Active AnimSet", "Show entry poses only for the active AnimSet", EUserInterfaceActionType::RadioButton, FInputChord());
+	UI_COMMAND(ShowEntryPosesAllSets, "All AnimSets", "Show entry poses for all the AnimSets", EUserInterfaceActionType::RadioButton, FInputChord());
+	UI_COMMAND(ShowEntryPosesNone, "None", "Hide Entry Poses", EUserInterfaceActionType::RadioButton, FInputChord());
+
 	UI_COMMAND(Simulate, "Simulate", "Simulate Mode", EUserInterfaceActionType::RadioButton, FInputChord());
 
 }
diff --git a/Engine/Plugins/Experimental/Animation/ContextualAnimation/Source/ContextualAnimationEditor/Private/ContextualAnimEdMode.cpp b/Engine/Plugins/Experimental/Animation/ContextualAnimation/Source/ContextualAnimationEditor/Private/ContextualAnimEdMode.cpp
--- a/Engine/Plugins/Experimental/Animation/ContextualAnimation/Source/ContextualAnimationEditor/Private/ContextualAnimEdMode.cpp
+++ b/Engine/Plugins/Experimental/Animation/ContextualAnimation/Source/ContextualAnimationEditor/Private/ContextualAnimEdMode.cpp
@@ -14,6 +14,8 @@
 #include "CanvasTypes.h"
 #include "CanvasItem.h"
 #include "Components/SkeletalMeshComponent.h"
+#include "Components/CapsuleComponent.h"
+#include "GameFramework/Character.h"
 
 const FEditorModeID FContextualAnimEdMode::EdModeId = TEXT("ContextualAnimEdMode");
 
@@ -48,10 +50,18 @@ void FContextualAnimEdMode::Render(const FSceneView* View, FViewport* Viewport,
 		if (const UContextualAnimSceneAsset* SceneAsset = ViewModel->GetSceneAsset())
 		{
 			const FContextualAnimSceneBindings& Bindings = ViewModel->SceneBindings;
+			const FName PrimaryRole = SceneAsset->GetPrimaryRole();
+			const FContextualAnimSceneBinding* PrimaryBinding = Bindings.FindBindingByRole(PrimaryRole);
+			
+			if(PrimaryBinding == nullptr)
+			{
+				return;
+			}
 
-			// Draw Scene Pivots
+			const FTransform PrimaryTransform = PrimaryBinding->GetContext().GetTransform();
 			if (const FContextualAnimSceneSection* Section = SceneAsset->GetSection(Bindings.GetSectionIdx()))
 			{
+				// Draw Scene Pivots
 				if (const FContextualAnimSet* AnimSet = Section->GetAnimSet(Bindings.GetAnimSetIdx()))
 				{
 					for (const FTransform& ScenePivot : AnimSet->ScenePivots)
@@ -59,80 +69,141 @@ void FContextualAnimEdMode::Render(const FSceneView* View, FViewport* Viewport,
 						DrawCoordinateSystem(PDI, ScenePivot.GetLocation(), ScenePivot.Rotator(), 50.f, SDPG_Foreground);
 					}
 				}
-			}
-
-			const EShowIKTargetsDrawMode IKTargetsDrawMode = ViewportClient->GetShowIKTargetsDrawMode();
-
-			for (const FContextualAnimSceneBinding& Binding : Bindings)
-			{
-				// Draw IK Targets
-				if (IKTargetsDrawMode == EShowIKTargetsDrawMode::All || (IKTargetsDrawMode == EShowIKTargetsDrawMode::Selected && Binding.GetActor() == ViewModel->GetSelectedActor()))
-				{
-					DrawIKTargetsForBinding(*PDI, Binding);
-				}
 
-				// Draw Selection Criteria
-				if (const FContextualAnimSceneBinding* PrimaryBinding = Bindings.FindBindingByRole(SceneAsset->GetPrimaryRole()))
+				// Attempt to draw selection criteria and entry poses only when playback time is at 0 and simulate mode is not playing
+				const float PlaybackTime = ViewModel->GetPlaybackTime();
+				if (PlaybackTime <= 0.f && !ViewModel->IsSimulateModePlaying())
 				{
-					const FTransform PrimaryTransform = PrimaryBinding->GetTransform();
-
-					const FContextualAnimTrack& AnimTrack = Binding.GetAnimTrack();
-					for (int32 CriterionIdx = 0; CriterionIdx < AnimTrack.SelectionCriteria.Num(); CriterionIdx++)
+					for (int32 AnimSetIdx = 0; AnimSetIdx < Section->GetNumAnimSets(); AnimSetIdx++)
 					{
-						if (const UContextualAnimSelectionCriterion* Criterion = AnimTrack.SelectionCriteria[CriterionIdx])
+						const bool bIsActiveSet = AnimSetIdx == Bindings.GetAnimSetIdx();
+
+						for (const FContextualAnimTrack& AnimTrack : Section->GetAnimSet(AnimSetIdx)->Tracks)
 						{
-							FLinearColor DrawColor = FLinearColor::White;
-							if (Criterion->DoesQuerierPassCondition(FContextualAnimSceneBindingContext(PrimaryTransform), Binding.GetContext()))
+							const FContextualAnimSceneBinding* Binding = Bindings.FindBindingByRole(AnimTrack.Role);
+
+							if (Binding == nullptr)
 							{
-								DrawColor = FLinearColor::Green;
+								continue;
 							}
 
-							//@TODO: Each SelectionCriterion should implement this, and here we should just call "Criterion->Draw()"
-							if (const UContextualAnimSelectionCriterion_TriggerArea* Spatial = Cast<UContextualAnimSelectionCriterion_TriggerArea>(Criterion))
+							bool bPassesSelectionCriteria = true;
+							for (int32 CriterionIdx = 0; CriterionIdx < AnimTrack.SelectionCriteria.Num(); CriterionIdx++)
 							{
-								const float HalfHeight = Spatial->Height / 2.f;
-								const int32 LastIndex = Spatial->PolygonPoints.Num() - 1;
-								for (int32 Idx = 0; Idx <= LastIndex; Idx++)
+								if (const UContextualAnimSelectionCriterion* Criterion = AnimTrack.SelectionCriteria[CriterionIdx])
 								{
-									const FVector P0 = PrimaryTransform.TransformPositionNoScale(Spatial->PolygonPoints[Idx]);
-									const FVector P1 = PrimaryTransform.TransformPositionNoScale(Spatial->PolygonPoints[Idx == LastIndex ? 0 : Idx + 1]);
-
-									PDI->DrawLine(P0, P1, DrawColor, SDPG_Foreground, 2.f);
-
-									PDI->DrawLine(P0 + FVector::UpVector * Spatial->Height, P1 + FVector::UpVector * Spatial->Height, DrawColor, SDPG_Foreground, 2.f);
-
-									PDI->DrawLine(P0, P0 + FVector::UpVector * Spatial->Height, DrawColor, SDPG_Foreground, 2.f);
-
-									PDI->SetHitProxy(new HSelectionCriterionHitProxy(Binding.GetRoleDef().Name, CriterionIdx, Idx));
-									PDI->DrawPoint(P0, FLinearColor::Black, 15.f, SDPG_Foreground);
-									PDI->SetHitProxy(nullptr);
-
-									PDI->SetHitProxy(new HSelectionCriterionHitProxy(Binding.GetRoleDef().Name, CriterionIdx, Idx + 4));
-									PDI->DrawPoint(P0 + FVector::UpVector * Spatial->Height, FLinearColor::Black, 15.f, SDPG_Foreground);
-									PDI->SetHitProxy(nullptr);
+									const bool bPassesCurrentSelectionCriterion = Criterion->DoesQuerierPassCondition(PrimaryBinding->GetContext(), Binding->GetContext());
+									if (!bPassesCurrentSelectionCriterion && bPassesSelectionCriteria)
+									{
+										bPassesSelectionCriteria = false;
+									}
+
+									// Draw Selection Criteria
+									const EMultiOptionDrawMode SelectionCriteriaDrawMode = ViewportClient->GetSelectionCriteriaDrawMode();
+									if (SelectionCriteriaDrawMode == EMultiOptionDrawMode::All || (SelectionCriteriaDrawMode == EMultiOptionDrawMode::Single && bIsActiveSet))
+									{
+										// Draw with dashed lines if this is not the active AnimSet and simulate mode is inactive
+										const bool bUseDashedLines = !bIsActiveSet && ViewModel->IsSimulateModeInactive();
+
+										FLinearColor DrawColor = FLinearColor::White;
+										float Thickness = 1.f;
+										if (bPassesCurrentSelectionCriterion)
+										{
+											DrawColor = FLinearColor::Green;
+											Thickness = 2.f;
+										}
+
+										//@TODO: Each SelectionCriterion should implement this, and here we should just call "Criterion->Draw()"
+										if (const UContextualAnimSelectionCriterion_TriggerArea* Spatial = Cast<UContextualAnimSelectionCriterion_TriggerArea>(Criterion))
+										{
+											const int32 LastIndex = Spatial->PolygonPoints.Num() - 1;
+											for (int32 Idx = 0; Idx <= LastIndex; Idx++)
+											{
+												const FVector P0 = PrimaryTransform.TransformPositionNoScale(Spatial->PolygonPoints[Idx]);
+												const FVector P1 = PrimaryTransform.TransformPositionNoScale(Spatial->PolygonPoints[Idx == LastIndex ? 0 : Idx + 1]);
+
+												if (bUseDashedLines)
+												{
+													DrawDashedLine(PDI, P0, P1, DrawColor, 10.f, SDPG_Foreground);
+													DrawDashedLine(PDI, P0 + FVector::UpVector * Spatial->Height, P1 + FVector::UpVector * Spatial->Height, DrawColor, 10.f, SDPG_Foreground);
+													DrawDashedLine(PDI, P0, P0 + FVector::UpVector * Spatial->Height, DrawColor, 10.f, SDPG_Foreground);
+												}
+												else
+												{
+													PDI->DrawLine(P0, P1, DrawColor, SDPG_Foreground, Thickness);
+													PDI->DrawLine(P0 + FVector::UpVector * Spatial->Height, P1 + FVector::UpVector * Spatial->Height, DrawColor, SDPG_Foreground, Thickness);
+													PDI->DrawLine(P0, P0 + FVector::UpVector * Spatial->Height, DrawColor, SDPG_Foreground, Thickness);
+												}
+
+												// Draw hit proxy points only for the active set and if simulate mode is not active
+												if (ViewModel->IsSimulateModeInactive() && bIsActiveSet)
+												{
+													PDI->SetHitProxy(new HSelectionCriterionHitProxy(AnimTrack.Role, CriterionIdx, Idx));
+													PDI->DrawPoint(P0, FLinearColor::Black, 15.f, SDPG_Foreground);
+													PDI->SetHitProxy(nullptr);
+
+													PDI->SetHitProxy(new HSelectionCriterionHitProxy(AnimTrack.Role, CriterionIdx, Idx + 4));
+													PDI->DrawPoint(P0 + FVector::UpVector * Spatial->Height, FLinearColor::Black, 15.f, SDPG_Foreground);
+													PDI->SetHitProxy(nullptr);
+												}
+											}
+										}
+										else if (const UContextualAnimSelectionCriterion_Cone* Cone = Cast<UContextualAnimSelectionCriterion_Cone>(Criterion))
+										{
+											FVector Origin, Direction;
+											if (Cone->Mode == EContextualAnimCriterionConeMode::ToPrimary)
+											{
+												const FTransform QuerierTransform = Section->GetAlignmentTransformForRoleRelativeToOtherRole(AnimSetIdx, AnimTrack.Role, PrimaryRole, 0.f) * PrimaryTransform;
+												Origin = QuerierTransform.GetLocation();
+												Direction = QuerierTransform.GetRotation().GetForwardVector().RotateAngleAxis(Cone->Offset, FVector::UpVector);
+											}
+											else if (Cone->Mode == EContextualAnimCriterionConeMode::FromPrimary)
+											{
+												Origin = PrimaryTransform.GetLocation();
+												Direction = PrimaryTransform.GetRotation().GetForwardVector().RotateAngleAxis(Cone->Offset, FVector::UpVector);
+											}
+
+											UContextualAnimUtilities::DrawSector(*PDI, Origin, Direction, 0.f, Cone->Distance, -Cone->HalfAngle, Cone->HalfAngle, DrawColor, SDPG_World, Thickness, bUseDashedLines);
+										}
+									}
 								}
 							}
-							else if (const UContextualAnimSelectionCriterion_Cone* Cone = Cast<UContextualAnimSelectionCriterion_Cone>(Criterion))
-							{
-								const FTransform QuerierTransform = Binding.GetContext().GetTransform();
 
-								FVector Origin, Direction;
-								if (Cone->Mode == EContextualAnimCriterionConeMode::ToPrimary)
-								{
-									Origin = QuerierTransform.GetLocation();
-									Direction = QuerierTransform.GetRotation().GetForwardVector().RotateAngleAxis(Cone->Offset, FVector::UpVector);
-								}
-								else if (Cone->Mode == EContextualAnimCriterionConeMode::FromPrimary)
+							// Draw Entry Poses
+							const EMultiOptionDrawMode EntryPosesDrawMode = ViewportClient->GetEntryPosesDrawMode();
+							if (EntryPosesDrawMode == EMultiOptionDrawMode::All || (EntryPosesDrawMode == EMultiOptionDrawMode::Single && bIsActiveSet))
+							{
+								if (AnimTrack.Role != PrimaryRole && AnimTrack.Animation)
 								{
-									Origin = PrimaryTransform.GetLocation();
-									Direction = PrimaryTransform.GetRotation().GetForwardVector().RotateAngleAxis(Cone->Offset, FVector::UpVector);
-								}
+									const float Time = 0.f;
+									const FTransform RelativeToPrimary = Section->GetAlignmentTransformForRoleRelativeToOtherRole(AnimSetIdx, AnimTrack.Role, PrimaryRole, Time);
+									FTransform Transform = (SceneAsset->GetMeshToComponentForRole(AnimTrack.Role) * RelativeToPrimary) * PrimaryTransform;
 
-								UContextualAnimUtilities::DrawSector(*PDI, Origin, Direction, 0.f, Cone->Distance, -Cone->HalfAngle, Cone->HalfAngle, DrawColor, SDPG_World, 1.f);
+									if (ACharacter* PrimaryAsCharacter = Cast<ACharacter>(PrimaryBinding->GetActor()))
+									{
+										const float CapsuleHalfHeight = PrimaryAsCharacter->GetCapsuleComponent()->GetScaledCapsuleHalfHeight();
+										Transform.SetLocation(Transform.GetLocation() - (FVector::UpVector * CapsuleHalfHeight));
+									}
+
+									UContextualAnimUtilities::DrawPose(PDI, AnimTrack.Animation, Time, Transform, bPassesSelectionCriteria ? FColor::Green : FColor::White, 1.f);
+								}
 							}
 						}
 					}
 				}
+
+				// Draw IK Targets
+				const EMultiOptionDrawMode IKTargetsDrawMode = ViewportClient->GetIKTargetsDrawMode();
+				if (IKTargetsDrawMode != EMultiOptionDrawMode::None)
+				{
+					for (const FContextualAnimSceneBinding& Binding : Bindings)
+					{
+						if (IKTargetsDrawMode == EMultiOptionDrawMode::All || (IKTargetsDrawMode == EMultiOptionDrawMode::Single && Binding.GetActor() == ViewModel->GetSelectedActor()))
+						{
+							DrawIKTargetsForBinding(*PDI, Binding);
+						}
+					}
+				}
 			}
 		}
 	}
diff --git a/Engine/Plugins/Experimental/Animation/ContextualAnimation/Source/ContextualAnimationEditor/Private/ContextualAnimViewModel.cpp b/Engine/Plugins/Experimental/Animation/ContextualAnimation/Source/ContextualAnimationEditor/Private/ContextualAnimViewModel.cpp
--- a/Engine/Plugins/Experimental/Animation/ContextualAnimation/Source/ContextualAnimationEditor/Private/ContextualAnimViewModel.cpp
+++ b/Engine/Plugins/Experimental/Animation/ContextualAnimation/Source/ContextualAnimationEditor/Private/ContextualAnimViewModel.cpp
@@ -112,6 +112,11 @@ void FContextualAnimViewModel::CreateSequencer()
 	Sequencer->SetPlaybackStatus(EMovieScenePlayerStatus::Stopped);
 }
 
+float FContextualAnimViewModel::GetPlaybackTime() const
+{
+	return Sequencer->GetGlobalTime().AsSeconds();
+}
+
 void FContextualAnimViewModel::SetActiveSection(int32 SectionIdx)
 {
 	check(GetSceneAsset()->Sections.IsValidIndex(SectionIdx));
diff --git a/Engine/Plugins/Experimental/Animation/ContextualAnimation/Source/ContextualAnimationEditor/Private/ContextualAnimViewportClient.cpp b/Engine/Plugins/Experimental/Animation/ContextualAnimation/Source/ContextualAnimationEditor/Private/ContextualAnimViewportClient.cpp
--- a/Engine/Plugins/Experimental/Animation/ContextualAnimation/Source/ContextualAnimationEditor/Private/ContextualAnimViewportClient.cpp
+++ b/Engine/Plugins/Experimental/Animation/ContextualAnimation/Source/ContextualAnimationEditor/Private/ContextualAnimViewportClient.cpp
@@ -37,17 +37,4 @@ void FContextualAnimViewportClient::TrackingStopped()
 void FContextualAnimViewportClient::Draw(const FSceneView* View, FPrimitiveDrawInterface* PDI)
 {
 	FEditorViewportClient::Draw(View, PDI);
-}
-
-void FContextualAnimViewportClient::OnSetIKTargetsDrawMode(EShowIKTargetsDrawMode Mode)
-{
-	if(ShowIKTargetsDrawMode != Mode)
-	{
-		ShowIKTargetsDrawMode = Mode;
-	}
-}
-
-bool FContextualAnimViewportClient::IsIKTargetsDrawModeSet(EShowIKTargetsDrawMode Mode) const
-{
-	return ShowIKTargetsDrawMode == Mode;
 }
\ No newline at end of file
diff --git a/Engine/Plugins/Experimental/Animation/ContextualAnimation/Source/ContextualAnimationEditor/Private/SContextualAnimViewport.cpp b/Engine/Plugins/Experimental/Animation/ContextualAnimation/Source/ContextualAnimationEditor/Private/SContextualAnimViewport.cpp
--- a/Engine/Plugins/Experimental/Animation/ContextualAnimation/Source/ContextualAnimationEditor/Private/SContextualAnimViewport.cpp
+++ b/Engine/Plugins/Experimental/Animation/ContextualAnimation/Source/ContextualAnimationEditor/Private/SContextualAnimViewport.cpp
@@ -180,23 +180,68 @@ void SContextualAnimViewport::BindCommands()
 
 	TSharedRef<FContextualAnimViewportClient> ViewportClientRef = ViewportClient.ToSharedRef();
 
-	CommandList->MapAction(
-		Commands.ShowIKTargetsDrawAll,
-		FExecuteAction::CreateSP(ViewportClientRef, &FContextualAnimViewportClient::OnSetIKTargetsDrawMode, EShowIKTargetsDrawMode::All),
-		FCanExecuteAction(),
-		FIsActionChecked::CreateSP(ViewportClientRef, &FContextualAnimViewportClient::IsIKTargetsDrawModeSet, EShowIKTargetsDrawMode::All));
-
-	CommandList->MapAction(
-		Commands.ShowIKTargetsDrawSelected,
-		FExecuteAction::CreateSP(ViewportClientRef, &FContextualAnimViewportClient::OnSetIKTargetsDrawMode, EShowIKTargetsDrawMode::Selected),
-		FCanExecuteAction(),
-		FIsActionChecked::CreateSP(ViewportClientRef, &FContextualAnimViewportClient::IsIKTargetsDrawModeSet, EShowIKTargetsDrawMode::Selected));
-
-	CommandList->MapAction(
-		Commands.ShowIKTargetsDrawNone,
-		FExecuteAction::CreateSP(ViewportClientRef, &FContextualAnimViewportClient::OnSetIKTargetsDrawMode, EShowIKTargetsDrawMode::None),
-		FCanExecuteAction(),
-		FIsActionChecked::CreateSP(ViewportClientRef, &FContextualAnimViewportClient::IsIKTargetsDrawModeSet, EShowIKTargetsDrawMode::None));
+	// Show IK Targets Options
+	{
+		CommandList->MapAction(
+			Commands.ShowIKTargetsDrawAll,
+			FExecuteAction::CreateSP(ViewportClientRef, &FContextualAnimViewportClient::SetIKTargetsDrawMode, EMultiOptionDrawMode::All),
+			FCanExecuteAction(),
+			FIsActionChecked::CreateSP(ViewportClientRef, &FContextualAnimViewportClient::IsIKTargetsDrawModeSet, EMultiOptionDrawMode::All));
+
+		CommandList->MapAction(
+			Commands.ShowIKTargetsDrawSelected,
+			FExecuteAction::CreateSP(ViewportClientRef, &FContextualAnimViewportClient::SetIKTargetsDrawMode, EMultiOptionDrawMode::Single),
+			FCanExecuteAction(),
+			FIsActionChecked::CreateSP(ViewportClientRef, &FContextualAnimViewportClient::IsIKTargetsDrawModeSet, EMultiOptionDrawMode::Single));
+
+		CommandList->MapAction(
+			Commands.ShowIKTargetsDrawNone,
+			FExecuteAction::CreateSP(ViewportClientRef, &FContextualAnimViewportClient::SetIKTargetsDrawMode, EMultiOptionDrawMode::None),
+			FCanExecuteAction(),
+			FIsActionChecked::CreateSP(ViewportClientRef, &FContextualAnimViewportClient::IsIKTargetsDrawModeSet, EMultiOptionDrawMode::None));
+	}
+
+	// Show Selection Criteria Options
+	{
+		CommandList->MapAction(
+			Commands.ShowSelectionCriteriaAllSets,
+			FExecuteAction::CreateSP(ViewportClientRef, &FContextualAnimViewportClient::SetSelectionCriteriaDrawMode, EMultiOptionDrawMode::All),
+			FCanExecuteAction(),
+			FIsActionChecked::CreateSP(ViewportClientRef, &FContextualAnimViewportClient::IsSelectionCriteriaDrawModeSet, EMultiOptionDrawMode::All));
+
+		CommandList->MapAction(
+			Commands.ShowSelectionCriteriaActiveSet,
+			FExecuteAction::CreateSP(ViewportClientRef, &FContextualAnimViewportClient::SetSelectionCriteriaDrawMode, EMultiOptionDrawMode::Single),
+			FCanExecuteAction(),
+			FIsActionChecked::CreateSP(ViewportClientRef, &FContextualAnimViewportClient::IsSelectionCriteriaDrawModeSet, EMultiOptionDrawMode::Single));
+
+		CommandList->MapAction(
+			Commands.ShowSelectionCriteriaNone,
+			FExecuteAction::CreateSP(ViewportClientRef, &FContextualAnimViewportClient::SetSelectionCriteriaDrawMode, EMultiOptionDrawMode::None),
+			FCanExecuteAction(),
+			FIsActionChecked::CreateSP(ViewportClientRef, &FContextualAnimViewportClient::IsSelectionCriteriaDrawModeSet, EMultiOptionDrawMode::None));
+	}
+
+	// Show Entry Poses Options
+	{
+		CommandList->MapAction(
+			Commands.ShowEntryPosesAllSets,
+			FExecuteAction::CreateSP(ViewportClientRef, &FContextualAnimViewportClient::SetEntryPosesDrawMode, EMultiOptionDrawMode::All),
+			FCanExecuteAction(),
+			FIsActionChecked::CreateSP(ViewportClientRef, &FContextualAnimViewportClient::IsEntryPosesDrawModeSet, EMultiOptionDrawMode::All));
+
+		CommandList->MapAction(
+			Commands.ShowEntryPosesActiveSet,
+			FExecuteAction::CreateSP(ViewportClientRef, &FContextualAnimViewportClient::SetEntryPosesDrawMode, EMultiOptionDrawMode::Single),
+			FCanExecuteAction(),
+			FIsActionChecked::CreateSP(ViewportClientRef, &FContextualAnimViewportClient::IsEntryPosesDrawModeSet, EMultiOptionDrawMode::Single));
+
+		CommandList->MapAction(
+			Commands.ShowEntryPosesNone,
+			FExecuteAction::CreateSP(ViewportClientRef, &FContextualAnimViewportClient::SetEntryPosesDrawMode, EMultiOptionDrawMode::None),
+			FCanExecuteAction(),
+			FIsActionChecked::CreateSP(ViewportClientRef, &FContextualAnimViewportClient::IsEntryPosesDrawModeSet, EMultiOptionDrawMode::None));
+	}
 }
 
 TSharedRef<FEditorViewportClient> SContextualAnimViewport::MakeEditorViewportClient()
diff --git a/Engine/Plugins/Experimental/Animation/ContextualAnimation/Source/ContextualAnimationEditor/Private/SContextualAnimViewportToolbar.cpp b/Engine/Plugins/Experimental/Animation/ContextualAnimation/Source/ContextualAnimationEditor/Private/SContextualAnimViewportToolbar.cpp
--- a/Engine/Plugins/Experimental/Animation/ContextualAnimation/Source/ContextualAnimationEditor/Private/SContextualAnimViewportToolbar.cpp
+++ b/Engine/Plugins/Experimental/Animation/ContextualAnimation/Source/ContextualAnimationEditor/Private/SContextualAnimViewportToolbar.cpp
@@ -37,6 +37,40 @@ TSharedRef<SWidget> SContextualAnimViewportToolBar::GenerateShowMenu() const
 					SubMenuBuilder.EndSection();
 				})
 		);
+
+		ShowMenuBuilder.AddSubMenu(
+			LOCTEXT("ShowMenu_SelectionCriteriaDrawSubMenu", "Selection Criteria"),
+			LOCTEXT("ShowMenu_SelectionCriteriaDrawSubMenuToolTip", "Selection Criteria Drawing Options"),
+			FNewMenuDelegate::CreateLambda([](FMenuBuilder& SubMenuBuilder)
+				{
+					const FContextualAnimAssetEditorCommands& Commands = FContextualAnimAssetEditorCommands::Get();
+
+					SubMenuBuilder.BeginSection("Selection Criteria", LOCTEXT("ShowMenu_SelectionCriteriaLabel", "Selection Criteria"));
+					{
+						SubMenuBuilder.AddMenuEntry(Commands.ShowSelectionCriteriaActiveSet);
+						SubMenuBuilder.AddMenuEntry(Commands.ShowSelectionCriteriaAllSets);
+						SubMenuBuilder.AddMenuEntry(Commands.ShowSelectionCriteriaNone);
+					}
+					SubMenuBuilder.EndSection();
+				})
+		);
+
+		ShowMenuBuilder.AddSubMenu(
+			LOCTEXT("ShowMenu_EntryPosesDrawSubMenu", "Entry Poses"),
+			LOCTEXT("ShowMenu_EntryPosesDrawSubMenuToolTip", "Entry Poses Drawing Options"),
+			FNewMenuDelegate::CreateLambda([](FMenuBuilder& SubMenuBuilder)
+				{
+					const FContextualAnimAssetEditorCommands& Commands = FContextualAnimAssetEditorCommands::Get();
+
+					SubMenuBuilder.BeginSection("Entry Poses", LOCTEXT("ShowMenu_EntryPosesLabel", "Entry Poses"));
+					{
+						SubMenuBuilder.AddMenuEntry(Commands.ShowEntryPosesActiveSet);
+						SubMenuBuilder.AddMenuEntry(Commands.ShowEntryPosesAllSets);
+						SubMenuBuilder.AddMenuEntry(Commands.ShowEntryPosesNone);
+					}
+					SubMenuBuilder.EndSection();
+				})
+		);
 	}
 
 	return ShowMenuBuilder.MakeWidget();
diff --git a/Engine/Plugins/Experimental/Animation/ContextualAnimation/Source/ContextualAnimationEditor/Public/ContextualAnimAssetEditorCommands.h b/Engine/Plugins/Experimental/Animation/ContextualAnimation/Source/ContextualAnimationEditor/Public/ContextualAnimAssetEditorCommands.h
--- a/Engine/Plugins/Experimental/Animation/ContextualAnimation/Source/ContextualAnimationEditor/Public/ContextualAnimAssetEditorCommands.h
+++ b/Engine/Plugins/Experimental/Animation/ContextualAnimation/Source/ContextualAnimationEditor/Public/ContextualAnimAssetEditorCommands.h
@@ -23,10 +23,16 @@ public:
 	TSharedPtr<FUICommandInfo> NewAnimSet;
 
 	TSharedPtr<FUICommandInfo> ShowIKTargetsDrawSelected;
-
 	TSharedPtr<FUICommandInfo> ShowIKTargetsDrawAll;
-
 	TSharedPtr<FUICommandInfo> ShowIKTargetsDrawNone;
 
+	TSharedPtr<FUICommandInfo> ShowSelectionCriteriaActiveSet;
+	TSharedPtr<FUICommandInfo> ShowSelectionCriteriaAllSets;
+	TSharedPtr<FUICommandInfo> ShowSelectionCriteriaNone;
+
+	TSharedPtr<FUICommandInfo> ShowEntryPosesActiveSet;
+	TSharedPtr<FUICommandInfo> ShowEntryPosesAllSets;
+	TSharedPtr<FUICommandInfo> ShowEntryPosesNone;
+
 	TSharedPtr<FUICommandInfo> Simulate;
 };
diff --git a/Engine/Plugins/Experimental/Animation/ContextualAnimation/Source/ContextualAnimationEditor/Public/ContextualAnimViewModel.h b/Engine/Plugins/Experimental/Animation/ContextualAnimation/Source/ContextualAnimationEditor/Public/ContextualAnimViewModel.h
--- a/Engine/Plugins/Experimental/Animation/ContextualAnimation/Source/ContextualAnimationEditor/Public/ContextualAnimViewModel.h
+++ b/Engine/Plugins/Experimental/Animation/ContextualAnimation/Source/ContextualAnimationEditor/Public/ContextualAnimViewModel.h
@@ -118,6 +118,8 @@ public:
 	void ApplyChangeToActorTransformInScene();
 	void DiscardChangeToActorTransformInScene();
 
+	float GetPlaybackTime() const;
+
 private:
 
 	struct FContextualAnimEdSelectionInfo
diff --git a/Engine/Plugins/Experimental/Animation/ContextualAnimation/Source/ContextualAnimationEditor/Public/ContextualAnimViewportClient.h b/Engine/Plugins/Experimental/Animation/ContextualAnimation/Source/ContextualAnimationEditor/Public/ContextualAnimViewportClient.h
--- a/Engine/Plugins/Experimental/Animation/ContextualAnimation/Source/ContextualAnimationEditor/Public/ContextualAnimViewportClient.h
+++ b/Engine/Plugins/Experimental/Animation/ContextualAnimation/Source/ContextualAnimationEditor/Public/ContextualAnimViewportClient.h
@@ -9,10 +9,10 @@ class FContextualAnimPreviewScene;
 class SContextualAnimViewport;
 class FContextualAnimAssetEditorToolkit;
 
-enum class EShowIKTargetsDrawMode : uint8
+enum class EMultiOptionDrawMode : uint8
 {
 	None,
-	Selected,
+	Single,
 	All
 };
 
@@ -34,9 +34,17 @@ public:
 
 	TSharedRef<FContextualAnimAssetEditorToolkit> GetAssetEditorToolkit() const { return AssetEditorToolkitPtr.Pin().ToSharedRef(); }
 
-	void OnSetIKTargetsDrawMode(EShowIKTargetsDrawMode Mode);
-	bool IsIKTargetsDrawModeSet(EShowIKTargetsDrawMode Mode) const;
-	EShowIKTargetsDrawMode GetShowIKTargetsDrawMode() const { return ShowIKTargetsDrawMode; }
+	void SetIKTargetsDrawMode(EMultiOptionDrawMode Mode) { IKTargetsDrawMode = Mode; }
+	bool IsIKTargetsDrawModeSet(EMultiOptionDrawMode Mode) const { return IKTargetsDrawMode == Mode; }
+	EMultiOptionDrawMode GetIKTargetsDrawMode() const { return IKTargetsDrawMode; }
+
+	void SetSelectionCriteriaDrawMode(EMultiOptionDrawMode Mode) { SelectionCriteriaDrawMode = Mode; }
+	bool IsSelectionCriteriaDrawModeSet(EMultiOptionDrawMode Mode) const { return SelectionCriteriaDrawMode == Mode; }
+	EMultiOptionDrawMode GetSelectionCriteriaDrawMode() const { return SelectionCriteriaDrawMode; }
+
+	void SetEntryPosesDrawMode(EMultiOptionDrawMode Mode) { EntryPosesDrawMode = Mode; }
+	bool IsEntryPosesDrawModeSet(EMultiOptionDrawMode Mode) const { return EntryPosesDrawMode == Mode; }
+	EMultiOptionDrawMode GetEntryPosesDrawMode() const { return EntryPosesDrawMode; }
 
 private:
 
@@ -46,5 +54,9 @@ private:
 	/** Asset editor toolkit we are embedded in */
 	TWeakPtr<FContextualAnimAssetEditorToolkit> AssetEditorToolkitPtr;
 
-	EShowIKTargetsDrawMode ShowIKTargetsDrawMode = EShowIKTargetsDrawMode::None;
-};
+	EMultiOptionDrawMode IKTargetsDrawMode = EMultiOptionDrawMode::None;
+
+	EMultiOptionDrawMode SelectionCriteriaDrawMode = EMultiOptionDrawMode::All;
+
+	EMultiOptionDrawMode EntryPosesDrawMode = EMultiOptionDrawMode::All;
+};
\ No newline at end of file
