Contextual Anim: Added helper function to facilitate the creation of bindings for two actors interaction

diff --git a/Engine/Plugins/Experimental/Animation/ContextualAnimation/Source/ContextualAnimation/Private/ContextualAnimTypes.cpp b/Engine/Plugins/Experimental/Animation/ContextualAnimation/Source/ContextualAnimation/Private/ContextualAnimTypes.cpp
--- a/Engine/Plugins/Experimental/Animation/ContextualAnimation/Source/ContextualAnimation/Private/ContextualAnimTypes.cpp
+++ b/Engine/Plugins/Experimental/Animation/ContextualAnimation/Source/ContextualAnimation/Private/ContextualAnimTypes.cpp
@@ -392,6 +392,61 @@ bool FContextualAnimSceneBindings::TryCreateBindings(const UContextualAnimSceneA
 	return (OutBindings.Num() == SceneAsset.GetNumRoles());
 }
 
+bool FContextualAnimSceneBindings::TryCreateBindings(const UContextualAnimSceneAsset& SceneAsset, int32 SectionIdx, int32 AnimSetIdx, const FContextualAnimSceneBindingContext& Primary, const FContextualAnimSceneBindingContext& Secondary, FContextualAnimSceneBindings& OutBindings)
+{
+	check(SceneAsset.HasValidData());
+
+	OutBindings.Reset();
+
+	const TArray<FContextualAnimRoleDefinition>& Roles = SceneAsset.GetRolesAsset()->Roles;
+
+	if (Roles.Num() > 2)
+	{
+		UE_LOG(LogContextualAnim, Warning, TEXT("FContextualAnimSceneBindings::TryCreateBindings Failed. Reason: Trying to create bindings with two actors for a SceneAsset with more than two roles. SceneAsset: %s Num Roles: %d SectionIdx: %d AnimSetIdx: %d"),
+			*GetNameSafe(&SceneAsset), Roles.Num(), SectionIdx, AnimSetIdx);
+
+		return false;
+	}
+
+	for(const FContextualAnimRoleDefinition& RoleDef : Roles)
+	{
+		const FName PrimaryRole = SceneAsset.GetPrimaryRole();
+		if(RoleDef.Name == PrimaryRole)
+		{
+			const FContextualAnimTrack* AnimTrack = SceneAsset.GetAnimTrack(SectionIdx, AnimSetIdx, PrimaryRole);
+			if (AnimTrack && AnimTrack->DoesQuerierPassSelectionCriteria(Primary, Primary))
+			{
+				OutBindings.Add(FContextualAnimSceneBinding(Primary, SceneAsset, *AnimTrack));
+			}
+			else
+			{
+				UE_LOG(LogContextualAnim, Warning, TEXT("FContextualAnimSceneBindings::TryCreateBindings Failed. Reason: Can't find valid track for primary actor. SceneAsset: %s Role: %s Actor: %s SectionIdx: %d AnimSetIdx: %d"),
+					*GetNameSafe(&SceneAsset), *RoleDef.Name.ToString(), *GetNameSafe(Primary.GetActor()), SectionIdx, AnimSetIdx);
+
+				return false;
+			}
+		}
+		else // Secondary Role
+		{
+			const FContextualAnimTrack* AnimTrack = SceneAsset.GetAnimTrack(SectionIdx, AnimSetIdx, RoleDef.Name);
+			if (AnimTrack && AnimTrack->DoesQuerierPassSelectionCriteria(Primary, Secondary))
+			{
+				OutBindings.Add(FContextualAnimSceneBinding(Secondary, SceneAsset, *AnimTrack));
+			}
+			else
+			{
+				UE_LOG(LogContextualAnim, Warning, TEXT("FContextualAnimSceneBindings::TryCreateBindings Failed. Reason: Can't find valid track for secondary actor. SceneAsset: %s Role: %s Actor: %s SectionIdx: %d AnimSetIdx: %d"),
+					*GetNameSafe(&SceneAsset), *RoleDef.Name.ToString(), *GetNameSafe(Secondary.GetActor()), SectionIdx, AnimSetIdx);
+
+				return false;
+			}
+		}
+	}
+
+	// Success if all the roles were filled
+	return (OutBindings.Num() == SceneAsset.GetNumRoles());
+}
+
 void FContextualAnimSceneBindings::CalculateAnimSetPivots(TArray<FContextualAnimSetPivot>& OutScenePivots) const
 {
 	if (const UContextualAnimSceneAsset* SceneAsset = GetSceneAsset())
diff --git a/Engine/Plugins/Experimental/Animation/ContextualAnimation/Source/ContextualAnimation/Private/ContextualAnimUtilities.cpp b/Engine/Plugins/Experimental/Animation/ContextualAnimation/Source/ContextualAnimation/Private/ContextualAnimUtilities.cpp
--- a/Engine/Plugins/Experimental/Animation/ContextualAnimation/Source/ContextualAnimation/Private/ContextualAnimUtilities.cpp
+++ b/Engine/Plugins/Experimental/Animation/ContextualAnimation/Source/ContextualAnimation/Private/ContextualAnimUtilities.cpp
@@ -337,6 +337,28 @@ bool UContextualAnimUtilities::BP_CreateContextualAnimSceneBindings(const UConte
 	return false;
 }
 
+bool UContextualAnimUtilities::BP_CreateContextualAnimSceneBindingsForTwoActors(const UContextualAnimSceneAsset* SceneAsset, const FContextualAnimSceneBindingContext& Primary, const FContextualAnimSceneBindingContext& Secondary, FContextualAnimSceneBindings& OutBindings)
+{
+	if (SceneAsset == nullptr || !SceneAsset->HasValidData())
+	{
+		UE_LOG(LogContextualAnim, Warning, TEXT("UContextualAnimUtilities::BP_CreateContextualAnimSceneBindingsForTwoActors Failed. Reason: Invalid or Empty SceneAsset (%s)"), *GetNameSafe(SceneAsset));
+		return false;
+	}
+
+	const int32 SectionIdx = 0; // Always start from the first section
+	const int32 NumSets = SceneAsset->GetNumAnimSetsInSection(SectionIdx);
+	for (int32 AnimSetIdx = 0; AnimSetIdx < NumSets; AnimSetIdx++)
+	{
+		OutBindings.Reset();
+		if (FContextualAnimSceneBindings::TryCreateBindings(*SceneAsset, SectionIdx, AnimSetIdx, Primary, Secondary, OutBindings))
+		{
+			return true;
+		}
+	}
+
+	return false;
+}
+
 // SceneBindings Blueprint Interface
 //------------------------------------------------------------------------------------------
 
diff --git a/Engine/Plugins/Experimental/Animation/ContextualAnimation/Source/ContextualAnimation/Public/ContextualAnimTypes.h b/Engine/Plugins/Experimental/Animation/ContextualAnimation/Source/ContextualAnimation/Public/ContextualAnimTypes.h
--- a/Engine/Plugins/Experimental/Animation/ContextualAnimation/Source/ContextualAnimation/Public/ContextualAnimTypes.h
+++ b/Engine/Plugins/Experimental/Animation/ContextualAnimation/Source/ContextualAnimation/Public/ContextualAnimTypes.h
@@ -402,6 +402,8 @@ struct CONTEXTUALANIMATION_API FContextualAnimSceneBindings
 
 	static bool TryCreateBindings(const UContextualAnimSceneAsset& SceneAsset, int32 SectionIdx, int32 AnimSetIdx, const TMap<FName, FContextualAnimSceneBindingContext>& Params, FContextualAnimSceneBindings& OutBindings);
 
+	static bool TryCreateBindings(const UContextualAnimSceneAsset& SceneAsset, int32 SectionIdx, int32 AnimSetIdx, const FContextualAnimSceneBindingContext& Primary, const FContextualAnimSceneBindingContext& Secondary, FContextualAnimSceneBindings& OutBindings);
+
 	void CalculateAnimSetPivots(TArray<FContextualAnimSetPivot>& OutScenePivots) const;
 	bool CalculateAnimSetPivot(const FContextualAnimSetPivotDefinition& AnimSetPivotDef, FContextualAnimSetPivot& OutScenePivot) const;
 
diff --git a/Engine/Plugins/Experimental/Animation/ContextualAnimation/Source/ContextualAnimation/Public/ContextualAnimUtilities.h b/Engine/Plugins/Experimental/Animation/ContextualAnimation/Source/ContextualAnimation/Public/ContextualAnimUtilities.h
--- a/Engine/Plugins/Experimental/Animation/ContextualAnimation/Source/ContextualAnimation/Public/ContextualAnimUtilities.h
+++ b/Engine/Plugins/Experimental/Animation/ContextualAnimation/Source/ContextualAnimation/Public/ContextualAnimUtilities.h
@@ -63,6 +63,9 @@ public:
 	UFUNCTION(BlueprintCallable, Category = "Contextual Anim|Scene Asset", meta = (DisplayName = "Create Contextual Anim Scene Bindings"))
 	static bool BP_CreateContextualAnimSceneBindings(const UContextualAnimSceneAsset* SceneAsset, const TMap<FName, FContextualAnimSceneBindingContext>& Params, FContextualAnimSceneBindings& OutBindings);
 
+	UFUNCTION(BlueprintCallable, Category = "Contextual Anim|Scene Asset", meta = (DisplayName = "Create Contextual Anim Scene Bindings For Two Actors"))
+	static bool BP_CreateContextualAnimSceneBindingsForTwoActors(const UContextualAnimSceneAsset* SceneAsset, const FContextualAnimSceneBindingContext& Primary, const FContextualAnimSceneBindingContext& Secondary, FContextualAnimSceneBindings& OutBindings);
+
 	// Montage Blueprint Interface
 	//------------------------------------------------------------------------------------------
 
