Contextual Anim: Fixed bug where autogenerated ik tracks could be wrong if ForceRootLock is set in the animations and fixed bug where ik tracks container was allocating memory for num bones entries instead of num tracks

diff --git a/Engine/Plugins/Experimental/Animation/ContextualAnimation/Source/ContextualAnimation/Private/ContextualAnimSceneAsset.cpp b/Engine/Plugins/Experimental/Animation/ContextualAnimation/Source/ContextualAnimation/Private/ContextualAnimSceneAsset.cpp
--- a/Engine/Plugins/Experimental/Animation/ContextualAnimation/Source/ContextualAnimation/Private/ContextualAnimSceneAsset.cpp
+++ b/Engine/Plugins/Experimental/Animation/ContextualAnimation/Source/ContextualAnimation/Private/ContextualAnimSceneAsset.cpp
@@ -25,6 +25,31 @@ static FCompactPoseBoneIndex GetCompactPoseBoneIndexFromPose(const FCSPose<FComp
 	return FCompactPoseBoneIndex(INDEX_NONE);
 }
 
+static void ExtractPoseIgnoringForceRootLock(UAnimSequenceBase* AnimSequenceBase, const FBoneContainer& BoneContainer, float Time, bool bExtractRootMotion, FCSPose<FCompactPose>& OutPose)
+{
+	UAnimSequence* AnimSequence = nullptr;
+	if (UAnimMontage* AnimMontage = Cast<UAnimMontage>(AnimSequenceBase))
+	{
+		if (AnimMontage->SlotAnimTracks.Num() > 0)
+		{
+			if (FAnimSegment* Segment = AnimMontage->SlotAnimTracks[0].AnimTrack.GetSegmentAtTime(Time))
+			{
+				AnimSequence = Cast<UAnimSequence>(Segment->GetAnimReference());
+			}
+		}
+	}
+	else
+	{
+		AnimSequence = Cast<UAnimSequence>(AnimSequenceBase);
+	}
+
+	if (AnimSequence)
+	{
+		TGuardValue<bool> ForceRootLockGuard(AnimSequence->bForceRootLock, false);
+		UContextualAnimUtilities::ExtractComponentSpacePose(AnimSequenceBase, BoneContainer, Time, bExtractRootMotion, OutPose);
+	}
+}
+
 // FContextualAnimSceneSection
 //==============================================================================================
 
@@ -157,7 +182,7 @@ void FContextualAnimSceneSection::GenerateIKTargetTracks(UContextualAnimSceneAss
 				continue;
 			}
 
-			if (const UAnimSequenceBase* Animation = AnimTrack.Animation)
+			if (UAnimSequenceBase* Animation = AnimTrack.Animation)
 			{
 				UE_LOG(LogContextualAnim, Log, TEXT("Generating IK Target Tracks. Animation: %s"), *GetNameSafe(Animation));
 
@@ -236,6 +261,8 @@ void FContextualAnimSceneSection::GenerateIKTargetTracks(UContextualAnimSceneAss
 
 					// Initialize track container
 					AnimTrack.IKTargetData.Initialize(TotalTracks, SampleInterval);
+					AnimTrack.IKTargetData.Tracks.TrackNames.AddZeroed(TotalTracks);
+					AnimTrack.IKTargetData.Tracks.AnimationTracks.AddZeroed(TotalTracks);
 
 					float Time = 0.f;
 					float EndTime = Animation->GetPlayLength();
@@ -247,7 +274,7 @@ void FContextualAnimSceneSection::GenerateIKTargetTracks(UContextualAnimSceneAss
 
 						// Extract pose from my animation
 						FCSPose<FCompactPose> ComponentSpacePose;
-						UContextualAnimUtilities::ExtractComponentSpacePose(Animation, BoneContainer, Time, false, ComponentSpacePose);
+						ExtractPoseIgnoringForceRootLock(Animation, BoneContainer, Time, false, ComponentSpacePose);
 
 						// For each target role
 						for (auto& Data : PoseExtractionHelperMap)
@@ -256,7 +283,7 @@ void FContextualAnimSceneSection::GenerateIKTargetTracks(UContextualAnimSceneAss
 							FCSPose<FCompactPose> OtherComponentSpacePose;
 							TArray<FBoneIndexType> OtherRequiredBoneIndexArray;
 							FBoneContainer OtherBoneContainer;
-							const UAnimSequenceBase* OtherAnimation = Data.Value.TargetAnimTrackPtr->Animation;
+							UAnimSequenceBase* OtherAnimation = Data.Value.TargetAnimTrackPtr->Animation;
 							if (OtherAnimation)
 							{
 								// Prepare array with the indices of the bones to extract from target animation
@@ -277,15 +304,14 @@ void FContextualAnimSceneSection::GenerateIKTargetTracks(UContextualAnimSceneAss
 									OtherBoneContainer = FBoneContainer(OtherRequiredBoneIndexArray, FCurveEvaluationOption(false), *OtherAnimation->GetSkeleton());
 
 									// Extract pose from target animation
-									UContextualAnimUtilities::ExtractComponentSpacePose(OtherAnimation, OtherBoneContainer, Time, false, OtherComponentSpacePose);
+									ExtractPoseIgnoringForceRootLock(OtherAnimation, OtherBoneContainer, Time, false, OtherComponentSpacePose);
 								}
 							}
 
 							for (int32 Idx = 0; Idx < Data.Value.BonesData.Num(); Idx++)
 							{
 								const FName TrackName = Data.Value.BonesData[Idx].Get<0>();
-								AnimTrack.IKTargetData.Tracks.TrackNames.Add(TrackName);
-								AnimTrack.IKTargetData.Tracks.AnimationTracks.AddZeroed();
+								AnimTrack.IKTargetData.Tracks.TrackNames[Idx] = TrackName;
 
 								// Get bone transform from my animation
 								const FName BoneName = Data.Value.BonesData[Idx].Get<1>();
