Contextual Anim: Fix for possible crash when generating ik tracks due to out-of-bound time.

diff --git a/Engine/Plugins/Experimental/Animation/ContextualAnimation/Source/ContextualAnimation/Private/ContextualAnimSceneAsset.cpp b/Engine/Plugins/Experimental/Animation/ContextualAnimation/Source/ContextualAnimation/Private/ContextualAnimSceneAsset.cpp
--- a/Engine/Plugins/Experimental/Animation/ContextualAnimation/Source/ContextualAnimation/Private/ContextualAnimSceneAsset.cpp
+++ b/Engine/Plugins/Experimental/Animation/ContextualAnimation/Source/ContextualAnimation/Private/ContextualAnimSceneAsset.cpp
@@ -32,7 +32,8 @@ static void ExtractPoseIgnoringForceRootLock(UAnimSequenceBase* AnimSequenceBase
 	{
 		if (AnimMontage->SlotAnimTracks.Num() > 0)
 		{
-			if (FAnimSegment* Segment = AnimMontage->SlotAnimTracks[0].AnimTrack.GetSegmentAtTime(Time))
+			const float ClampedTime = FMath::Clamp(Time, 0.f, AnimSequenceBase->GetPlayLength());
+			if (FAnimSegment* Segment = AnimMontage->SlotAnimTracks[0].AnimTrack.GetSegmentAtTime(ClampedTime))
 			{
 				AnimSequence = Cast<UAnimSequence>(Segment->GetAnimReference());
 			}
