Contextual Anim: Initial pass for adding net serialize to binding data

diff --git a/Engine/Plugins/Experimental/Animation/ContextualAnimation/Source/ContextualAnimation/Private/ContextualAnimSceneInstance.cpp b/Engine/Plugins/Experimental/Animation/ContextualAnimation/Source/ContextualAnimation/Private/ContextualAnimSceneInstance.cpp
--- a/Engine/Plugins/Experimental/Animation/ContextualAnimation/Source/ContextualAnimation/Private/ContextualAnimSceneInstance.cpp
+++ b/Engine/Plugins/Experimental/Animation/ContextualAnimation/Source/ContextualAnimation/Private/ContextualAnimSceneInstance.cpp
@@ -176,7 +176,7 @@ float UContextualAnimSceneInstance::TransitionTo(FContextualAnimSceneBinding& Bi
 	AnimInstance->OnMontageBlendingOut.RemoveDynamic(this, &UContextualAnimSceneInstance::OnMontageBlendingOut);
 
 	const UAnimMontage* Montage = PlayAnimation(*AnimInstance, *AnimTrack.Animation);
-	Binding.AnimTrackPtr = &AnimTrack;
+	Binding.SetAnimTrack(AnimTrack);
 
 	AnimInstance->OnMontageBlendingOut.AddUniqueDynamic(this, &UContextualAnimSceneInstance::OnMontageBlendingOut);
 
diff --git a/Engine/Plugins/Experimental/Animation/ContextualAnimation/Source/ContextualAnimation/Private/ContextualAnimTypes.cpp b/Engine/Plugins/Experimental/Animation/ContextualAnimation/Source/ContextualAnimation/Private/ContextualAnimTypes.cpp
--- a/Engine/Plugins/Experimental/Animation/ContextualAnimation/Source/ContextualAnimation/Private/ContextualAnimTypes.cpp
+++ b/Engine/Plugins/Experimental/Animation/ContextualAnimation/Source/ContextualAnimation/Private/ContextualAnimTypes.cpp
@@ -19,6 +19,7 @@ const FContextualAnimSceneBinding FContextualAnimSceneBinding::InvalidBinding;
 const FContextualAnimTrack FContextualAnimTrack::EmptyTrack;
 const FContextualAnimIKTarget FContextualAnimIKTarget::InvalidIKTarget;
 const FContextualAnimIKTargetDefContainer FContextualAnimIKTargetDefContainer::EmptyContainer;
+const FContextualAnimRoleDefinition FContextualAnimRoleDefinition::InvalidRoleDefinition;
 
 // FContextualAnimAlignmentTrackContainer
 ///////////////////////////////////////////////////////////////////////
@@ -261,6 +262,14 @@ FVector FContextualAnimSceneBindingContext::GetVelocity() const
 	return FVector::ZeroVector;
 }
 
+bool FContextualAnimSceneBindingContext::NetSerialize(FArchive& Ar, class UPackageMap* Map, bool& bOutSuccess)
+{
+	//@TODO: Serialize optional properties
+	Ar << Actor;
+	bOutSuccess = true;
+	return true;
+}
+
 // FContextualAnimSceneBinding
 ///////////////////////////////////////////////////////////////////////
 
@@ -270,9 +279,36 @@ FContextualAnimSceneBinding::FContextualAnimSceneBinding(const FContextualAnimSc
 
 	Context = InContext;
 	SceneAsset = &InSceneAsset;
-	AnimTrackPtr = &InAnimTrack;
-	RoleDefPtr = InSceneAsset.GetRolesAsset()->FindRoleDefinitionByName(InAnimTrack.Role);
-	check(RoleDefPtr);
+	SectionIdx = InAnimTrack.SectionIdx;
+	AnimSetIdx = InAnimTrack.AnimSetIdx;
+	AnimTrackIdx = InAnimTrack.AnimTrackIdx;
+}
+
+const FContextualAnimTrack& FContextualAnimSceneBinding::GetAnimTrack() const
+{ 
+	if(const FContextualAnimTrack* AnimTrack = SceneAsset->GetAnimTrack(SectionIdx, AnimSetIdx, AnimTrackIdx))
+	{
+		return *AnimTrack;
+	}
+
+	return FContextualAnimTrack::EmptyTrack;
+}
+
+void FContextualAnimSceneBinding::SetAnimTrack(const FContextualAnimTrack& InAnimTrack) 
+{ 
+	SectionIdx = InAnimTrack.SectionIdx;
+	AnimSetIdx = InAnimTrack.AnimSetIdx;
+	AnimTrackIdx = InAnimTrack.AnimTrackIdx;
+}
+
+const FContextualAnimRoleDefinition& FContextualAnimSceneBinding::GetRoleDef() const
+{ 
+	if(const FContextualAnimRoleDefinition* RoleDef = SceneAsset->GetRolesAsset()->FindRoleDefinitionByName(GetAnimTrack().Role))
+	{
+		return *RoleDef;
+	}
+	
+	return FContextualAnimRoleDefinition::InvalidRoleDefinition;
 }
 
 UContextualAnimSceneActorComponent* FContextualAnimSceneBinding::GetSceneActorComponent() const
@@ -330,9 +366,26 @@ const FContextualAnimIKTargetDefContainer& FContextualAnimSceneBinding::GetIKTar
 	return GetSceneAsset().GetIKTargetDefsForRoleInSection(GetAnimTrack().SectionIdx, GetAnimTrack().Role);
 }
 
+bool FContextualAnimSceneBinding::NetSerialize(FArchive& Ar, class UPackageMap* Map, bool& bOutSuccess)
+{
+	Context.NetSerialize(Ar, Map, bOutSuccess);
+	Ar << SceneAsset;
+	Ar << SectionIdx;
+	Ar << AnimSetIdx;
+	Ar << AnimTrackIdx;
+	bOutSuccess = true;
+	return true;
+}
+
 // FContextualAnimSceneBindings
 ///////////////////////////////////////////////////////////////////////
 
+bool FContextualAnimSceneBindings::NetSerialize(FArchive& Ar, class UPackageMap* Map, bool& bOutSuccess)
+{
+	bOutSuccess = SafeNetSerializeTArray_WithNetSerialize<10>(Ar, Data, Map);
+	return true;
+}
+
 bool FContextualAnimSceneBindings::TryCreateBindings(const UContextualAnimSceneAsset& SceneAsset, int32 SectionIdx, int32 AnimSetIdx, const TMap<FName, FContextualAnimSceneBindingContext>& Params, FContextualAnimSceneBindings& OutBindings)
 {
 	check(SceneAsset.HasValidData());
diff --git a/Engine/Plugins/Experimental/Animation/ContextualAnimation/Source/ContextualAnimation/Public/ContextualAnimTypes.h b/Engine/Plugins/Experimental/Animation/ContextualAnimation/Source/ContextualAnimation/Public/ContextualAnimTypes.h
--- a/Engine/Plugins/Experimental/Animation/ContextualAnimation/Source/ContextualAnimation/Public/ContextualAnimTypes.h
+++ b/Engine/Plugins/Experimental/Animation/ContextualAnimation/Source/ContextualAnimation/Public/ContextualAnimTypes.h
@@ -215,13 +215,15 @@ struct FContextualAnimRoleDefinition
 	GENERATED_BODY()
 
 	UPROPERTY(EditAnywhere, BlueprintReadWrite, Category = "Defaults")
-	FName Name;
+	FName Name = NAME_None;
 
 	UPROPERTY(EditAnywhere, BlueprintReadWrite, Category = "Defaults")
 	bool bIsCharacter = false;
 
 	UPROPERTY(EditAnywhere, BlueprintReadWrite, Category = "Defaults")
 	FTransform MeshToComponent = FTransform(FRotator(0.f, -90.f, 0.f));
+
+	static const FContextualAnimRoleDefinition InvalidRoleDefinition;
 };
 
 // FContextualAnimSetPivotDefinition
@@ -294,6 +296,8 @@ struct CONTEXTUALANIMATION_API FContextualAnimSceneBindingContext
 
 	//@TODO: Add accessors for GameplayTags
 
+	bool NetSerialize(FArchive& Ar, class UPackageMap* Map, bool& bOutSuccess);
+
 private:
 
 	TWeakObjectPtr<AActor> Actor = nullptr;
@@ -305,6 +309,15 @@ private:
 	FGameplayTagContainer ExternalGameplayTags;
 };
 
+template<>
+struct TStructOpsTypeTraits<FContextualAnimSceneBindingContext> : public TStructOpsTypeTraitsBase2<FContextualAnimSceneBindingContext>
+{
+	enum
+	{
+		WithNetSerializer = true
+	};
+};
+
 /** Represent an actor bound to a role in the scene */
 USTRUCT(BlueprintType)
 struct CONTEXTUALANIMATION_API FContextualAnimSceneBinding
@@ -321,11 +334,12 @@ struct CONTEXTUALANIMATION_API FContextualAnimSceneBinding
 	FORCEINLINE FVector GetVelocity()  const { return GetContext().GetVelocity(); }
 
 	/** Returns the track in the scene asset used by this actor */
-	FORCEINLINE const FContextualAnimTrack& GetAnimTrack() const { check(AnimTrackPtr);  return *AnimTrackPtr; }
-	FORCEINLINE void SetAnimTrack(const FContextualAnimTrack& InAnimTrack) { AnimTrackPtr = &InAnimTrack; }
+	const FContextualAnimTrack& GetAnimTrack() const;
+	
+	void SetAnimTrack(const FContextualAnimTrack& InAnimTrack);
 
 	/** Returns the role definition used by this actor */
-	FORCEINLINE const FContextualAnimRoleDefinition& GetRoleDef() const { check(RoleDefPtr); return *RoleDefPtr; }
+	const FContextualAnimRoleDefinition& GetRoleDef() const;
 
 	/** Returns the SceneInstance we belong to (if any) */
 	FORCEINLINE const UContextualAnimSceneInstance* GetSceneInstance() const { return SceneInstancePtr.Get(); }
@@ -351,6 +365,8 @@ struct CONTEXTUALANIMATION_API FContextualAnimSceneBinding
 
 	UContextualAnimSceneActorComponent* GetSceneActorComponent() const;
 
+	bool NetSerialize(FArchive& Ar, class UPackageMap* Map, bool& bOutSuccess);
+
 	static const FContextualAnimSceneBinding InvalidBinding;
 
 private:
@@ -362,14 +378,21 @@ private:
 	/** Scene Asset we are bound to */
 	TWeakObjectPtr<const UContextualAnimSceneAsset> SceneAsset = nullptr;
 
-	/** Ptr to the animation data in the scene asset used by this actor */
-	const FContextualAnimTrack* AnimTrackPtr = nullptr;
-
-	/** Ptr to the role definition used by this actor */
-	const FContextualAnimRoleDefinition* RoleDefPtr = nullptr;
-
 	/** Ptr back to the scene instance we belong to (if any) */
 	TWeakObjectPtr<const UContextualAnimSceneInstance> SceneInstancePtr = nullptr;
+
+	int32 SectionIdx = INDEX_NONE;
+	int32 AnimSetIdx = INDEX_NONE;
+	int32 AnimTrackIdx = INDEX_NONE;
+};
+
+template<>
+struct TStructOpsTypeTraits<FContextualAnimSceneBinding> : public TStructOpsTypeTraitsBase2<FContextualAnimSceneBinding>
+{
+	enum
+	{
+		WithNetSerializer = true
+	};
 };
 
 USTRUCT(BlueprintType)
@@ -407,6 +430,8 @@ struct CONTEXTUALANIMATION_API FContextualAnimSceneBindings
 	void CalculateAnimSetPivots(TArray<FContextualAnimSetPivot>& OutScenePivots) const;
 	bool CalculateAnimSetPivot(const FContextualAnimSetPivotDefinition& AnimSetPivotDef, FContextualAnimSetPivot& OutScenePivot) const;
 
+	bool NetSerialize(FArchive& Ar, class UPackageMap* Map, bool& bOutSuccess);
+
 private:
 
 	/** List of actors bound to each role in the SceneAsset */
@@ -414,6 +439,15 @@ private:
 	TArray<FContextualAnimSceneBinding> Data;
 };
 
+template<>
+struct TStructOpsTypeTraits<FContextualAnimSceneBindings> : public TStructOpsTypeTraitsBase2<FContextualAnimSceneBindings>
+{
+	enum
+	{
+		WithNetSerializer = true
+	};
+};
+
 USTRUCT(BlueprintType)
 struct CONTEXTUALANIMATION_API FContextualAnimStartSceneParams
 {
