Contextual Anim: Removed dependecy with SceneInstance and added support for replication so we can get ik targets on simulated proxies.

diff --git a/Engine/Plugins/Experimental/Animation/ContextualAnimation/Source/ContextualAnimation/Private/ContextualAnimSceneActorComponent.cpp b/Engine/Plugins/Experimental/Animation/ContextualAnimation/Source/ContextualAnimation/Private/ContextualAnimSceneActorComponent.cpp
--- a/Engine/Plugins/Experimental/Animation/ContextualAnimation/Source/ContextualAnimation/Private/ContextualAnimSceneActorComponent.cpp
+++ b/Engine/Plugins/Experimental/Animation/ContextualAnimation/Source/ContextualAnimation/Private/ContextualAnimSceneActorComponent.cpp
@@ -7,11 +7,11 @@
 #include "Components/SkeletalMeshComponent.h"
 #include "DrawDebugHelpers.h"
 #include "ContextualAnimManager.h"
-#include "ContextualAnimSceneInstance.h"
 #include "ContextualAnimSceneAsset.h"
 #include "ContextualAnimUtilities.h"
 #include "AnimNotifyState_IKWindow.h"
 #include "EngineUtils.h"
+#include "Net/UnrealNetwork.h"
 
 #if !(UE_BUILD_SHIPPING || UE_BUILD_TEST)
 TAutoConsoleVariable<int32> CVarContextualAnimIKDebug(TEXT("a.ContextualAnim.IK.Debug"), 0, TEXT("Draw Debug IK Targets"));
@@ -24,6 +24,44 @@ UContextualAnimSceneActorComponent::UContextualAnimSceneActorComponent(const FOb
 {
 	PrimaryComponentTick.bCanEverTick = false;
 	PrimaryComponentTick.bStartWithTickEnabled = false;
+	SetIsReplicatedByDefault(true);
+}
+
+void UContextualAnimSceneActorComponent::GetLifetimeReplicatedProps(TArray< FLifetimeProperty >& OutLifetimeProps) const
+{
+	Super::GetLifetimeReplicatedProps(OutLifetimeProps);
+	DOREPLIFETIME_CONDITION(UContextualAnimSceneActorComponent, Bindings, COND_SimulatedOnly);
+}
+
+void UContextualAnimSceneActorComponent::OnRep_Bindings()
+{
+	if(Bindings.Num() > 0)
+	{
+		if (const FContextualAnimSceneBinding* Binding = Bindings.FindBindingByActor(GetOwner()))
+		{
+			// Start listening to TickPose if we joined an scene where we need IK
+			if (Binding->GetIKTargetDefs().IKTargetDefs.Num() > 0)
+			{
+				USkeletalMeshComponent* SkelMeshComp = UContextualAnimUtilities::TryGetSkeletalMeshComponent(GetOwner());
+				if (SkelMeshComp && !SkelMeshComp->OnTickPose.IsBoundToObject(this))
+				{
+					SkelMeshComp->OnTickPose.AddUObject(this, &UContextualAnimSceneActorComponent::OnTickPose);
+				}
+			}
+
+			OnJoinedSceneDelegate.Broadcast(this);
+		}
+	}
+	else
+	{
+		USkeletalMeshComponent* SkelMeshComp = UContextualAnimUtilities::TryGetSkeletalMeshComponent(GetOwner());
+		if (SkelMeshComp && SkelMeshComp->OnTickPose.IsBoundToObject(this))
+		{
+			SkelMeshComp->OnTickPose.RemoveAll(this);
+		}
+
+		OnLeftSceneDelegate.Broadcast(this);
+	}
 }
 
 FBoxSphereBounds UContextualAnimSceneActorComponent::CalcBounds(const FTransform& LocalToWorld) const
@@ -57,39 +95,41 @@ void UContextualAnimSceneActorComponent::OnUnregister()
 	}
 }
 
-void UContextualAnimSceneActorComponent::OnJoinedScene(const FContextualAnimSceneBinding* Binding)
+void UContextualAnimSceneActorComponent::OnJoinedScene(const FContextualAnimSceneBindings& InBindings)
 {
-	check(Binding);
-	BindingPtr = Binding;
-
-	// Start listening to TickPose if we joined an scene where we need IK
-	if (Binding->GetIKTargetDefs().IKTargetDefs.Num() > 0)
+	if (const FContextualAnimSceneBinding* Binding = InBindings.FindBindingByActor(GetOwner()))
 	{
-		if (USkeletalMeshComponent* SkelMeshComp = UContextualAnimUtilities::TryGetSkeletalMeshComponent(GetOwner()))
+		Bindings = InBindings;
+
+		// Start listening to TickPose if we joined an scene where we need IK
+		if (Binding->GetIKTargetDefs().IKTargetDefs.Num() > 0)
 		{
-			SkelMeshComp->OnTickPose.AddUObject(this, &UContextualAnimSceneActorComponent::OnTickPose);
+			USkeletalMeshComponent* SkelMeshComp = UContextualAnimUtilities::TryGetSkeletalMeshComponent(GetOwner());
+			if (SkelMeshComp && !SkelMeshComp->OnTickPose.IsBoundToObject(this))
+			{
+				SkelMeshComp->OnTickPose.AddUObject(this, &UContextualAnimSceneActorComponent::OnTickPose);
+			}
 		}
-	}
 
-	OnJoinedSceneDelegate.Broadcast(this);
+		OnJoinedSceneDelegate.Broadcast(this);
+	}
 }
 
-void UContextualAnimSceneActorComponent::OnLeftScene(const FContextualAnimSceneBinding* Binding)
+void UContextualAnimSceneActorComponent::OnLeftScene()
 {
-	check(Binding);
-
-	OnLeftSceneDelegate.Broadcast(this);
-
-	// Stop listening to TickPose if we were
-	if (Binding->GetIKTargetDefs().IKTargetDefs.Num() > 0)
+	if (const FContextualAnimSceneBinding* Binding = Bindings.FindBindingByActor(GetOwner()))
 	{
-		if (USkeletalMeshComponent* SkelMeshComp = UContextualAnimUtilities::TryGetSkeletalMeshComponent(GetOwner()))
+		OnLeftSceneDelegate.Broadcast(this);
+
+		// Stop listening to TickPose if we were
+		USkeletalMeshComponent* SkelMeshComp = UContextualAnimUtilities::TryGetSkeletalMeshComponent(GetOwner());
+		if (SkelMeshComp && SkelMeshComp->OnTickPose.IsBoundToObject(this))
 		{
 			SkelMeshComp->OnTickPose.RemoveAll(this);
 		}
-	}
 
-	BindingPtr = nullptr;
+		Bindings.Reset();
+	}
 }
 
 void UContextualAnimSceneActorComponent::OnTickPose(class USkinnedMeshComponent* SkinnedMeshComponent, float DeltaTime, bool bNeedsValidRootMotion)
@@ -102,7 +142,11 @@ void UContextualAnimSceneActorComponent::UpdateIKTargets()
 {
 	IKTargets.Reset();
 
-	check(BindingPtr);
+	const FContextualAnimSceneBinding* BindingPtr = Bindings.FindBindingByActor(GetOwner());
+	if (BindingPtr == nullptr)
+	{
+		return;
+	}
 
 	const FAnimMontageInstance* MontageInstance = BindingPtr->GetAnimMontageInstance();
 	if(MontageInstance == nullptr)
@@ -125,7 +169,7 @@ void UContextualAnimSceneActorComponent::UpdateIKTargets()
 		// @TODO: IKTargetTransform will be off by 1 frame if we tick before target. 
 		// Should we at least add an option to the SceneAsset to setup tick dependencies or should this be entirely up to the user?
 
-		if (const FContextualAnimSceneBinding* TargetBinding = BindingPtr->GetSceneInstance()->FindBindingByRole(IKTargetDef.TargetRoleName))
+		if (const FContextualAnimSceneBinding* TargetBinding = Bindings.FindBindingByRole(IKTargetDef.TargetRoleName))
 		{
 			if (const USkeletalMeshComponent* TargetSkelMeshComp = TargetBinding->GetSkeletalMeshComponent())
 			{
diff --git a/Engine/Plugins/Experimental/Animation/ContextualAnimation/Source/ContextualAnimation/Private/ContextualAnimSceneInstance.cpp b/Engine/Plugins/Experimental/Animation/ContextualAnimation/Source/ContextualAnimation/Private/ContextualAnimSceneInstance.cpp
--- a/Engine/Plugins/Experimental/Animation/ContextualAnimation/Source/ContextualAnimation/Private/ContextualAnimSceneInstance.cpp
+++ b/Engine/Plugins/Experimental/Animation/ContextualAnimation/Source/ContextualAnimation/Private/ContextualAnimSceneInstance.cpp
@@ -140,7 +140,7 @@ float UContextualAnimSceneInstance::Join(FContextualAnimSceneBinding& Binding)
 
 	if (UContextualAnimSceneActorComponent* SceneActorComp = Binding.GetSceneActorComponent())
 	{
-		SceneActorComp->OnJoinedScene(&Binding);
+		SceneActorComp->OnJoinedScene(Bindings);
 	}
 
 	OnActorJoined.Broadcast(this, Actor);
@@ -305,7 +305,7 @@ void UContextualAnimSceneInstance::OnMontageBlendingOut(UAnimMontage* Montage, b
 
 			if (UContextualAnimSceneActorComponent* SceneActorComp = Binding.GetSceneActorComponent())
 			{
-				SceneActorComp->OnLeftScene(&Binding);
+				SceneActorComp->OnLeftScene();
 			}
 
 			OnActorLeft.Broadcast(this, Actor);
diff --git a/Engine/Plugins/Experimental/Animation/ContextualAnimation/Source/ContextualAnimation/Public/ContextualAnimSceneActorComponent.h b/Engine/Plugins/Experimental/Animation/ContextualAnimation/Source/ContextualAnimation/Public/ContextualAnimSceneActorComponent.h
--- a/Engine/Plugins/Experimental/Animation/ContextualAnimation/Source/ContextualAnimation/Public/ContextualAnimSceneActorComponent.h
+++ b/Engine/Plugins/Experimental/Animation/ContextualAnimation/Source/ContextualAnimation/Public/ContextualAnimSceneActorComponent.h
@@ -12,8 +12,6 @@ class AActor;
 class FPrimitiveSceneProxy;
 class UAnimInstance;
 class UAnimMontage;
-class UContextualAnimSceneInstance;
-struct FContextualAnimSceneBinding;
 
 DECLARE_DYNAMIC_MULTICAST_DELEGATE_OneParam(FContextualAnimSceneActorCompDelegate, class UContextualAnimSceneActorComponent*, SceneActorComponent);
 
@@ -40,6 +38,7 @@ public:
 
 	UContextualAnimSceneActorComponent(const FObjectInitializer& ObjectInitializer);
 
+	virtual void GetLifetimeReplicatedProps(TArray< FLifetimeProperty >& OutLifetimeProps) const;
 	virtual void OnRegister() override;
 	virtual void OnUnregister() override;
 	virtual FBoxSphereBounds CalcBounds(const FTransform& LocalToWorld) const;
@@ -47,11 +46,13 @@ public:
 
 	virtual void AddIKGoals_Implementation(TMap<FName, FIKRigGoal>& OutGoals) override;
 
-	/** Called from the scene instance when the actor owner of this component joins an scene */
-	void OnJoinedScene(const FContextualAnimSceneBinding* Binding);
+	/** Called when the actor owner of this component joins an scene */
+	UFUNCTION(BlueprintCallable, Category = "Contextual Anim|Scene Actor Component")
+	void OnJoinedScene(const FContextualAnimSceneBindings& InBindings);
 	
 	/** Called from the scene instance when the actor owner of this component leave an scene */
-	void OnLeftScene(const FContextualAnimSceneBinding* Binding);
+	UFUNCTION(BlueprintCallable, Category = "Contextual Anim|Scene Actor Component")
+	void OnLeftScene();
 
 	UFUNCTION(BlueprintPure, Category = "Contextual Anim|Scene Actor Component")
 	const TArray<FContextualAnimIKTarget>& GetIKTargets() const { return IKTargets; }
@@ -61,8 +62,8 @@ public:
 
 protected:
 
-	/** Ptr back to the binding that represent us in the scene instance we are part of */
-	const FContextualAnimSceneBinding* BindingPtr = nullptr;
+	UPROPERTY(Transient, ReplicatedUsing = OnRep_Bindings)
+	FContextualAnimSceneBindings Bindings;
 
 	/** List of IKTarget for this frame */
 	UPROPERTY()
@@ -77,6 +78,9 @@ protected:
 	UFUNCTION()
 	void OnTickPose(class USkinnedMeshComponent* SkinnedMeshComponent, float DeltaTime, bool bNeedsValidRootMotion);
 
+	UFUNCTION()
+	void OnRep_Bindings();
+
 private:
 
 	bool bRegistered = false;
