Contextual Anim: Do not reset IKTargets once they are collected by IK Rig so they can still be queried during that frame

diff --git a/Engine/Plugins/Experimental/Animation/ContextualAnimation/Source/ContextualAnimation/Private/ContextualAnimSceneActorComponent.cpp b/Engine/Plugins/Experimental/Animation/ContextualAnimation/Source/ContextualAnimation/Private/ContextualAnimSceneActorComponent.cpp
--- a/Engine/Plugins/Experimental/Animation/ContextualAnimation/Source/ContextualAnimation/Private/ContextualAnimSceneActorComponent.cpp
+++ b/Engine/Plugins/Experimental/Animation/ContextualAnimation/Source/ContextualAnimation/Private/ContextualAnimSceneActorComponent.cpp
@@ -244,8 +244,6 @@ void UContextualAnimSceneActorComponent::AddIKGoals_Implementation(TMap<FName, F
 		Goal.RotationSpace = EIKRigGoalSpace::World;
 		OutGoals.Add(Goal.Name, Goal);
 	}
-
-	IKTargets.Reset();
 }
 
 const FContextualAnimIKTarget& UContextualAnimSceneActorComponent::GetIKTargetByGoalName(FName GoalName) const
