Contextual Anim: Added accessors to binding data + exposed some helper functions to blueprint

diff --git a/Engine/Plugins/Experimental/Animation/ContextualAnimation/Source/ContextualAnimation/Private/ContextualAnimSceneAsset.cpp b/Engine/Plugins/Experimental/Animation/ContextualAnimation/Source/ContextualAnimation/Private/ContextualAnimSceneAsset.cpp
--- a/Engine/Plugins/Experimental/Animation/ContextualAnimation/Source/ContextualAnimation/Private/ContextualAnimSceneAsset.cpp
+++ b/Engine/Plugins/Experimental/Animation/ContextualAnimation/Source/ContextualAnimation/Private/ContextualAnimSceneAsset.cpp
@@ -664,6 +664,26 @@ const TArray<FContextualAnimSetPivotDefinition>& UContextualAnimSceneAsset::GetA
 	return Sections.IsValidIndex(SectionIdx) ? Sections[SectionIdx].GetAnimSetPivotDefinitions() : EmptyDefs;
 }
 
+FTransform UContextualAnimSceneAsset::GetIKTargetTransform(int32 SectionIdx, int32 AnimSetIdx, int32 AnimTrackIdx, const FName& TrackName, float Time) const
+{
+	if (const FContextualAnimTrack* AnimTrack = GetAnimTrack(SectionIdx, AnimSetIdx, AnimTrackIdx))
+	{
+		return AnimTrack->IKTargetData.ExtractTransformAtTime(TrackName, Time);
+	}
+
+	return FTransform::Identity;
+}
+
+FTransform UContextualAnimSceneAsset::GetAlignmentTransform(int32 SectionIdx, int32 AnimSetIdx, int32 AnimTrackIdx, const FName& TrackName, float Time) const
+{
+	if (const FContextualAnimTrack* AnimTrack = GetAnimTrack(SectionIdx, AnimSetIdx, AnimTrackIdx))
+	{
+		return AnimTrack->AlignmentData.ExtractTransformAtTime(TrackName, Time);
+	}
+
+	return FTransform::Identity;
+}
+
 static FContextualAnimPoint GetContextualAnimPoint(const FContextualAnimTrack& AnimTrack, const FTransform& ToWorldTransform, int32 SampleRate, EContextualAnimPointType Type)
 {
 	check(SampleRate > 0);
diff --git a/Engine/Plugins/Experimental/Animation/ContextualAnimation/Source/ContextualAnimation/Private/ContextualAnimTypes.cpp b/Engine/Plugins/Experimental/Animation/ContextualAnimation/Source/ContextualAnimation/Private/ContextualAnimTypes.cpp
--- a/Engine/Plugins/Experimental/Animation/ContextualAnimation/Source/ContextualAnimation/Private/ContextualAnimTypes.cpp
+++ b/Engine/Plugins/Experimental/Animation/ContextualAnimation/Source/ContextualAnimation/Private/ContextualAnimTypes.cpp
@@ -380,6 +380,42 @@ bool FContextualAnimSceneBinding::NetSerialize(FArchive& Ar, class UPackageMap*
 // FContextualAnimSceneBindings
 ///////////////////////////////////////////////////////////////////////
 
+bool FContextualAnimSceneBindings::IsValid() const
+{
+	const UContextualAnimSceneAsset* SceneAsset = GetSceneAsset();
+	return SceneAsset && SceneAsset->HasValidData() && Num() > 0;
+}
+
+const FContextualAnimTrack& FContextualAnimSceneBindings::GetAnimTrackFromBinding(const FContextualAnimSceneBinding& Binding) const
+{
+	if (const FContextualAnimTrack* AnimTrack = GetSceneAsset()->GetAnimTrack(Binding.SectionIdx, Binding.AnimSetIdx, Binding.AnimTrackIdx))
+	{
+		return *AnimTrack;
+	}
+
+	return FContextualAnimTrack::EmptyTrack;
+}
+
+const FContextualAnimIKTargetDefContainer& FContextualAnimSceneBindings::GetIKTargetDefContainerFromBinding(const FContextualAnimSceneBinding& Binding) const
+{
+	return GetSceneAsset()->GetIKTargetDefsForRoleInSection(GetSectionIdx(), GetRoleFromBinding(Binding));
+}
+
+FTransform FContextualAnimSceneBindings::GetIKTargetTransformFromBinding(const FContextualAnimSceneBinding& Binding, const FName& TrackName, float Time) const
+{
+	return GetSceneAsset()->GetIKTargetTransform(Binding.SectionIdx, Binding.AnimSetIdx, Binding.AnimTrackIdx, TrackName, Time);
+}
+
+FTransform FContextualAnimSceneBindings::GetAlignmentTransformFromBinding(const FContextualAnimSceneBinding& Binding, const FName& TrackName, float Time) const
+{
+	return GetSceneAsset()->GetAlignmentTransform(Binding.SectionIdx, Binding.AnimSetIdx, Binding.AnimTrackIdx, TrackName, Time);
+}
+
+const FName& FContextualAnimSceneBindings::GetRoleFromBinding(const FContextualAnimSceneBinding& Binding) const
+{
+	return GetAnimTrackFromBinding(Binding).Role;
+}
+
 bool FContextualAnimSceneBindings::NetSerialize(FArchive& Ar, class UPackageMap* Map, bool& bOutSuccess)
 {
 	bOutSuccess = SafeNetSerializeTArray_WithNetSerialize<10>(Ar, Data, Map);
diff --git a/Engine/Plugins/Experimental/Animation/ContextualAnimation/Source/ContextualAnimation/Private/ContextualAnimUtilities.cpp b/Engine/Plugins/Experimental/Animation/ContextualAnimation/Source/ContextualAnimation/Private/ContextualAnimUtilities.cpp
--- a/Engine/Plugins/Experimental/Animation/ContextualAnimation/Source/ContextualAnimation/Private/ContextualAnimUtilities.cpp
+++ b/Engine/Plugins/Experimental/Animation/ContextualAnimation/Source/ContextualAnimation/Private/ContextualAnimUtilities.cpp
@@ -395,25 +395,28 @@ void UContextualAnimUtilities::BP_SceneBindings_CalculateAnimSetPivots(const FCo
 
 void UContextualAnimUtilities::BP_SceneBindings_AddOrUpdateWarpTargetsForBindings(const FContextualAnimSceneBindings& Bindings)
 {
+	if(!Bindings.IsValid())
+	{
+		UE_LOG(LogContextualAnim, Warning, TEXT("UContextualAnimUtilities::BP_SceneBindings_AddOrUpdateWarpTargetsForBindings Failed. Reason: Invalid Bindings. SceneAsset: %s"), *GetNameSafe(Bindings.GetSceneAsset()));
+		return;
+	}
+
 	const UContextualAnimSceneAsset* SceneAsset = Bindings.GetSceneAsset();
-	if(ensureAlways(SceneAsset))
+	const int32 SectionIdx = Bindings.GetSectionIdx();
+	for (const FContextualAnimSetPivotDefinition& PivotDef : SceneAsset->GetAnimSetPivotDefinitionsInSection(SectionIdx))
 	{
-		const int32 SectionIdx = Bindings.GetSectionIdx();
-		for (const FContextualAnimSetPivotDefinition& PivotDef : SceneAsset->GetAnimSetPivotDefinitionsInSection(SectionIdx))
+		FContextualAnimSetPivot ScenePivot;
+		if (Bindings.CalculateAnimSetPivot(PivotDef, ScenePivot))
 		{
-			FContextualAnimSetPivot ScenePivot;
-			if (Bindings.CalculateAnimSetPivot(PivotDef, ScenePivot))
+			for (const FContextualAnimSceneBinding& Binding : Bindings)
 			{
-				for (const FContextualAnimSceneBinding& Binding : Bindings)
+				if (UMotionWarpingComponent* MotionWarpComp = Binding.GetActor()->FindComponentByClass<UMotionWarpingComponent>())
 				{
-					if (UMotionWarpingComponent* MotionWarpComp = Binding.GetActor()->FindComponentByClass<UMotionWarpingComponent>())
-					{
-						//@TODO: Cache this
-						const float Time = Binding.GetAnimTrack().GetSyncTimeForWarpSection(PivotDef.Name);
-						const FTransform TransformRelativeToScenePivot = Binding.GetAnimTrack().AlignmentData.ExtractTransformAtTime(PivotDef.Name, Time);
-						const FTransform WarpTarget = (TransformRelativeToScenePivot * ScenePivot.Transform);
-						MotionWarpComp->AddOrUpdateWarpTargetFromTransform(PivotDef.Name, WarpTarget);
-					}
+					//@TODO: Cache this
+					const float Time = Binding.GetAnimTrack().GetSyncTimeForWarpSection(PivotDef.Name);
+					const FTransform TransformRelativeToScenePivot = Binding.GetAnimTrack().AlignmentData.ExtractTransformAtTime(PivotDef.Name, Time);
+					const FTransform WarpTarget = (TransformRelativeToScenePivot * ScenePivot.Transform);
+					MotionWarpComp->AddOrUpdateWarpTargetFromTransform(PivotDef.Name, WarpTarget);
 				}
 			}
 		}
@@ -446,3 +449,41 @@ FTransform UContextualAnimUtilities::BP_SceneBindings_GetAlignmentTransformForRo
 
 	return Result;
 }
+
+const UAnimSequenceBase* UContextualAnimUtilities::BP_SceneBinding_GetAnimationFromBinding(const FContextualAnimSceneBindings& Bindings, const FContextualAnimSceneBinding& Binding)
+{
+	if (!Bindings.IsValid())
+	{
+		UE_LOG(LogContextualAnim, Warning, TEXT("UContextualAnimUtilities::BP_SceneBinding_GetAnimationFromBinding Failed. Reason: Invalid Bindings. SceneAsset: %s"), *GetNameSafe(Bindings.GetSceneAsset()));
+		return nullptr;
+	}
+
+	return Bindings.GetAnimTrackFromBinding(Binding).Animation;
+}
+
+FName UContextualAnimUtilities::BP_SceneBinding_GetRoleFromBinding(const FContextualAnimSceneBindings& Bindings, const FContextualAnimSceneBinding& Binding)
+{
+	if (!Bindings.IsValid())
+	{
+		UE_LOG(LogContextualAnim, Warning, TEXT("UContextualAnimUtilities::BP_SceneBinding_GetRoleFromBinding Failed. Reason: Invalid Bindings. SceneAsset: %s"), *GetNameSafe(Bindings.GetSceneAsset()));
+		return NAME_None;
+	}
+
+	return Bindings.GetAnimTrackFromBinding(Binding).Role;
+}
+
+FTransform UContextualAnimUtilities::BP_SceneBindings_GetAlignmentTransformFromBinding(const FContextualAnimSceneBindings& Bindings, const FContextualAnimSceneBinding& Binding, const FContextualAnimSetPivot& Pivot)
+{
+	if (!Bindings.IsValid())
+	{
+		UE_LOG(LogContextualAnim, Warning, TEXT("UContextualAnimUtilities::BP_SceneBindings_GetAlignmentTransformFromBinding Failed. Reason: Invalid Bindings. SceneAsset: %s"), *GetNameSafe(Bindings.GetSceneAsset()));
+		return FTransform::Identity;
+	}
+
+	const FContextualAnimTrack& AnimTrack = Bindings.GetAnimTrackFromBinding(Binding);
+
+	float StartTime, EndTime;
+	AnimTrack.GetStartAndEndTimeForWarpSection(Pivot.Name, StartTime, EndTime);
+
+	return AnimTrack.GetAlignmentTransformAtTime(EndTime) * Pivot.Transform;
+}
diff --git a/Engine/Plugins/Experimental/Animation/ContextualAnimation/Source/ContextualAnimation/Public/ContextualAnimSceneAsset.h b/Engine/Plugins/Experimental/Animation/ContextualAnimation/Source/ContextualAnimation/Public/ContextualAnimSceneAsset.h
--- a/Engine/Plugins/Experimental/Animation/ContextualAnimation/Source/ContextualAnimation/Public/ContextualAnimSceneAsset.h
+++ b/Engine/Plugins/Experimental/Animation/ContextualAnimation/Source/ContextualAnimation/Public/ContextualAnimSceneAsset.h
@@ -188,6 +188,10 @@ public:
  
  	const FContextualAnimTrack* GetAnimTrack(int32 SectionIdx, int32 AnimSetIdx, int32 AnimTrackIdx) const;
 
+	FTransform GetIKTargetTransform(int32 SectionIdx, int32 AnimSetIdx, int32 AnimTrackIdx, const FName& TrackName, float Time) const;
+
+	FTransform GetAlignmentTransform(int32 SectionIdx, int32 AnimSetIdx, int32 AnimTrackIdx, const FName& TrackName, float Time) const;
+
 	const FContextualAnimTrack* FindAnimTrackByAnimation(const UAnimSequenceBase* Animation) const;
 
 	const TArray<FContextualAnimSetPivotDefinition>& GetAnimSetPivotDefinitionsInSection(int32 SectionIdx) const;
diff --git a/Engine/Plugins/Experimental/Animation/ContextualAnimation/Source/ContextualAnimation/Public/ContextualAnimTypes.h b/Engine/Plugins/Experimental/Animation/ContextualAnimation/Source/ContextualAnimation/Public/ContextualAnimTypes.h
--- a/Engine/Plugins/Experimental/Animation/ContextualAnimation/Source/ContextualAnimation/Public/ContextualAnimTypes.h
+++ b/Engine/Plugins/Experimental/Animation/ContextualAnimation/Source/ContextualAnimation/Public/ContextualAnimTypes.h
@@ -372,6 +372,7 @@ struct CONTEXTUALANIMATION_API FContextualAnimSceneBinding
 private:
 
 	friend UContextualAnimSceneInstance;
+	friend struct FContextualAnimSceneBindings;
 
 	FContextualAnimSceneBindingContext Context;
 
@@ -430,6 +431,14 @@ struct CONTEXTUALANIMATION_API FContextualAnimSceneBindings
 	void CalculateAnimSetPivots(TArray<FContextualAnimSetPivot>& OutScenePivots) const;
 	bool CalculateAnimSetPivot(const FContextualAnimSetPivotDefinition& AnimSetPivotDef, FContextualAnimSetPivot& OutScenePivot) const;
 
+	const FContextualAnimTrack& GetAnimTrackFromBinding(const FContextualAnimSceneBinding& Binding) const;
+	const FName& GetRoleFromBinding(const FContextualAnimSceneBinding& Binding) const;
+	FTransform GetAlignmentTransformFromBinding(const FContextualAnimSceneBinding& Binding, const FName& TrackName, float Time) const;
+	const FContextualAnimIKTargetDefContainer& GetIKTargetDefContainerFromBinding(const FContextualAnimSceneBinding& Binding) const;
+	FTransform GetIKTargetTransformFromBinding(const FContextualAnimSceneBinding& Binding, const FName& TrackName, float Time) const;
+
+	bool IsValid() const;
+
 	bool NetSerialize(FArchive& Ar, class UPackageMap* Map, bool& bOutSuccess);
 
 private:
diff --git a/Engine/Plugins/Experimental/Animation/ContextualAnimation/Source/ContextualAnimation/Public/ContextualAnimUtilities.h b/Engine/Plugins/Experimental/Animation/ContextualAnimation/Source/ContextualAnimation/Public/ContextualAnimUtilities.h
--- a/Engine/Plugins/Experimental/Animation/ContextualAnimation/Source/ContextualAnimation/Public/ContextualAnimUtilities.h
+++ b/Engine/Plugins/Experimental/Animation/ContextualAnimation/Source/ContextualAnimation/Public/ContextualAnimUtilities.h
@@ -96,6 +96,12 @@ public:
 	UFUNCTION(BlueprintPure, Category = "Contextual Anim|Scene Bindings", meta = (DisplayName = "Get Binding By Actor"))
 	static const FContextualAnimSceneBinding& BP_SceneBindings_GetBindingByActor(const FContextualAnimSceneBindings& Bindings, const AActor* Actor);
 
+	UFUNCTION(BlueprintPure, Category = "Contextual Anim|Scene Bindings", meta = (DisplayName = "Get Animation From Binding"))
+	static const UAnimSequenceBase* BP_SceneBinding_GetAnimationFromBinding(const FContextualAnimSceneBindings& Bindings, const FContextualAnimSceneBinding& Binding);
+
+	UFUNCTION(BlueprintPure, Category = "Contextual Anim|Scene Bindings", meta = (DisplayName = "Get Role From Binding"))
+	static FName BP_SceneBinding_GetRoleFromBinding(const FContextualAnimSceneBindings& Bindings, const FContextualAnimSceneBinding& Binding);
+
 	UFUNCTION(BlueprintPure, Category = "Contextual Anim|Scene Bindings", meta = (DisplayName = "Get Scene Asset"))
 	static const UContextualAnimSceneAsset* BP_SceneBindings_GetSceneAsset(const FContextualAnimSceneBindings& Bindings) { return Bindings.GetSceneAsset(); }
 
@@ -108,6 +114,9 @@ public:
 	UFUNCTION(BlueprintPure, Category = "Contextual Anim|Scene Bindings", meta = (DisplayName = "Get Alignment Transform For Role Relative To Pivot"))
 	static FTransform BP_SceneBindings_GetAlignmentTransformForRoleRelativeToPivot(const FContextualAnimSceneBindings& Bindings, FName Role, const FContextualAnimSetPivot& Pivot, float Time);
 
+	UFUNCTION(BlueprintPure, Category = "Contextual Anim|Scene Bindings", meta = (DisplayName = "Get Alignment Transform From Binding"))
+	static FTransform BP_SceneBindings_GetAlignmentTransformFromBinding(const FContextualAnimSceneBindings& Bindings, const FContextualAnimSceneBinding& Binding, const FContextualAnimSetPivot& Pivot);
+
 	// FContextualAnimSceneBindingContext Blueprint Interface
 	//------------------------------------------------------------------------------------------
 	
