Contextual Anim: Added static helper functions to create Bindings without an explicit AnimSet

diff --git a/Engine/Plugins/Experimental/Animation/ContextualAnimation/Source/ContextualAnimation/Private/ContextualAnimTypes.cpp b/Engine/Plugins/Experimental/Animation/ContextualAnimation/Source/ContextualAnimation/Private/ContextualAnimTypes.cpp
--- a/Engine/Plugins/Experimental/Animation/ContextualAnimation/Source/ContextualAnimation/Private/ContextualAnimTypes.cpp
+++ b/Engine/Plugins/Experimental/Animation/ContextualAnimation/Source/ContextualAnimation/Private/ContextualAnimTypes.cpp
@@ -542,6 +542,34 @@ bool FContextualAnimSceneBindings::TryCreateBindings(const UContextualAnimSceneA
 	return true;
 }
 
+bool FContextualAnimSceneBindings::TryCreateBindings(const UContextualAnimSceneAsset& SceneAsset, int32 SectionIdx, const TMap<FName, FContextualAnimSceneBindingContext>& Params, FContextualAnimSceneBindings& OutBindings)
+{
+	const int32 NumSets = SceneAsset.GetNumAnimSetsInSection(SectionIdx);
+	for (int32 AnimSetIdx = 0; AnimSetIdx < NumSets; AnimSetIdx++)
+	{
+		if (TryCreateBindings(SceneAsset, SectionIdx, AnimSetIdx, Params, OutBindings))
+		{
+			return true;
+		}
+	}
+
+	return false;
+}
+
+bool FContextualAnimSceneBindings::TryCreateBindings(const UContextualAnimSceneAsset& SceneAsset, int32 SectionIdx, const FContextualAnimSceneBindingContext& Primary, const FContextualAnimSceneBindingContext& Secondary, FContextualAnimSceneBindings& OutBindings)
+{
+	const int32 NumSets = SceneAsset.GetNumAnimSetsInSection(SectionIdx);
+	for (int32 AnimSetIdx = 0; AnimSetIdx < NumSets; AnimSetIdx++)
+	{
+		if (TryCreateBindings(SceneAsset, SectionIdx, AnimSetIdx, Primary, Secondary, OutBindings))
+		{
+			return true;
+		}
+	}
+
+	return false;
+}
+
 void FContextualAnimSceneBindings::CalculateAnimSetPivots(TArray<FContextualAnimSetPivot>& OutScenePivots) const
 {
 	if (IsValid())
diff --git a/Engine/Plugins/Experimental/Animation/ContextualAnimation/Source/ContextualAnimation/Private/ContextualAnimUtilities.cpp b/Engine/Plugins/Experimental/Animation/ContextualAnimation/Source/ContextualAnimation/Private/ContextualAnimUtilities.cpp
--- a/Engine/Plugins/Experimental/Animation/ContextualAnimation/Source/ContextualAnimation/Private/ContextualAnimUtilities.cpp
+++ b/Engine/Plugins/Experimental/Animation/ContextualAnimation/Source/ContextualAnimation/Private/ContextualAnimUtilities.cpp
@@ -324,17 +324,7 @@ bool UContextualAnimUtilities::BP_CreateContextualAnimSceneBindings(const UConte
 	}
 
 	const int32 SectionIdx = 0; // Always start from the first section
-	const int32 NumSets = SceneAsset->GetNumAnimSetsInSection(SectionIdx);
-	for (int32 AnimSetIdx = 0; AnimSetIdx < NumSets; AnimSetIdx++)
-	{
-		OutBindings.Reset();
-		if(FContextualAnimSceneBindings::TryCreateBindings(*SceneAsset, SectionIdx, AnimSetIdx, Params, OutBindings))
-		{
-			return true;
-		}
-	}
-	
-	return false;
+	return FContextualAnimSceneBindings::TryCreateBindings(*SceneAsset, SectionIdx, Params, OutBindings);
 }
 
 bool UContextualAnimUtilities::BP_CreateContextualAnimSceneBindingsForTwoActors(const UContextualAnimSceneAsset* SceneAsset, const FContextualAnimSceneBindingContext& Primary, const FContextualAnimSceneBindingContext& Secondary, FContextualAnimSceneBindings& OutBindings)
@@ -346,17 +336,7 @@ bool UContextualAnimUtilities::BP_CreateContextualAnimSceneBindingsForTwoActors(
 	}
 
 	const int32 SectionIdx = 0; // Always start from the first section
-	const int32 NumSets = SceneAsset->GetNumAnimSetsInSection(SectionIdx);
-	for (int32 AnimSetIdx = 0; AnimSetIdx < NumSets; AnimSetIdx++)
-	{
-		OutBindings.Reset();
-		if (FContextualAnimSceneBindings::TryCreateBindings(*SceneAsset, SectionIdx, AnimSetIdx, Primary, Secondary, OutBindings))
-		{
-			return true;
-		}
-	}
-
-	return false;
+	return FContextualAnimSceneBindings::TryCreateBindings(*SceneAsset, SectionIdx, Primary, Secondary, OutBindings);
 }
 
 // SceneBindings Blueprint Interface
diff --git a/Engine/Plugins/Experimental/Animation/ContextualAnimation/Source/ContextualAnimation/Public/ContextualAnimTypes.h b/Engine/Plugins/Experimental/Animation/ContextualAnimation/Source/ContextualAnimation/Public/ContextualAnimTypes.h
--- a/Engine/Plugins/Experimental/Animation/ContextualAnimation/Source/ContextualAnimation/Public/ContextualAnimTypes.h
+++ b/Engine/Plugins/Experimental/Animation/ContextualAnimation/Source/ContextualAnimation/Public/ContextualAnimTypes.h
@@ -408,8 +408,9 @@ struct CONTEXTUALANIMATION_API FContextualAnimSceneBindings
 	FORCEINLINE TArray<FContextualAnimSceneBinding>::RangedForConstIteratorType end() const { return Data.end(); }
 
 	static bool TryCreateBindings(const UContextualAnimSceneAsset& SceneAsset, int32 SectionIdx, int32 AnimSetIdx, const TMap<FName, FContextualAnimSceneBindingContext>& Params, FContextualAnimSceneBindings& OutBindings);
-
+	static bool TryCreateBindings(const UContextualAnimSceneAsset& SceneAsset, int32 SectionIdx, const TMap<FName, FContextualAnimSceneBindingContext>& Params, FContextualAnimSceneBindings& OutBindings);
 	static bool TryCreateBindings(const UContextualAnimSceneAsset& SceneAsset, int32 SectionIdx, int32 AnimSetIdx, const FContextualAnimSceneBindingContext& Primary, const FContextualAnimSceneBindingContext& Secondary, FContextualAnimSceneBindings& OutBindings);
+	static bool TryCreateBindings(const UContextualAnimSceneAsset& SceneAsset, int32 SectionIdx, const FContextualAnimSceneBindingContext& Primary, const FContextualAnimSceneBindingContext& Secondary, FContextualAnimSceneBindings& OutBindings);
 
 	void CalculateAnimSetPivots(TArray<FContextualAnimSetPivot>& OutScenePivots) const;
 	bool CalculateAnimSetPivot(const FContextualAnimSetPivotDefinition& AnimSetPivotDef, FContextualAnimSetPivot& OutScenePivot) const;
