Contextual Anim: Added logic to scene actor comp to disable collision between actors, disable bAllowPhysicsRotationDuringAnimRootMotion and keep anims in-sync

diff --git a/Engine/Plugins/Experimental/Animation/ContextualAnimation/Source/ContextualAnimation/Private/ContextualAnimSceneActorComponent.cpp b/Engine/Plugins/Experimental/Animation/ContextualAnimation/Source/ContextualAnimation/Private/ContextualAnimSceneActorComponent.cpp
--- a/Engine/Plugins/Experimental/Animation/ContextualAnimation/Source/ContextualAnimation/Private/ContextualAnimSceneActorComponent.cpp
+++ b/Engine/Plugins/Experimental/Animation/ContextualAnimation/Source/ContextualAnimation/Private/ContextualAnimSceneActorComponent.cpp
@@ -4,6 +4,7 @@
 #include "Animation/AnimInstance.h"
 #include "Animation/AnimMontage.h"
 #include "GameFramework/Character.h"
+#include "GameFramework/CharacterMovementComponent.h"
 #include "ContextualAnimSelectionCriterion.h"
 #include "Components/SkeletalMeshComponent.h"
 #include "DrawDebugHelpers.h"
@@ -35,25 +36,25 @@ UContextualAnimSceneActorComponent::UContextualAnimSceneActorComponent(const FOb
 void UContextualAnimSceneActorComponent::GetLifetimeReplicatedProps(TArray< FLifetimeProperty >& OutLifetimeProps) const
 {
 	Super::GetLifetimeReplicatedProps(OutLifetimeProps);
-	DOREPLIFETIME_CONDITION(UContextualAnimSceneActorComponent, Bindings, COND_SimulatedOnly);
+	DOREPLIFETIME_CONDITION(UContextualAnimSceneActorComponent, RepBindings, COND_SimulatedOnly);
 }
 
 void UContextualAnimSceneActorComponent::OnRep_Bindings()
 {
-	if(Bindings.Num() > 0)
+	if (RepBindings.IsValid())
 	{
-		if (const FContextualAnimSceneBinding* Binding = Bindings.FindBindingByActor(GetOwner()))
+		if (const FContextualAnimSceneBinding* Binding = RepBindings.FindBindingByActor(GetOwner()))
 		{
-			// Start listening to TickPose if we joined an scene where we need IK
-			if (Bindings.GetIKTargetDefContainerFromBinding(*Binding).IKTargetDefs.Num() > 0)
+			LocalBindings = RepBindings;
+
+			USkeletalMeshComponent* SkelMeshComp = UContextualAnimUtilities::TryGetSkeletalMeshComponent(GetOwner());
+			if (SkelMeshComp && !SkelMeshComp->OnTickPose.IsBoundToObject(this))
 			{
-				USkeletalMeshComponent* SkelMeshComp = UContextualAnimUtilities::TryGetSkeletalMeshComponent(GetOwner());
-				if (SkelMeshComp && !SkelMeshComp->OnTickPose.IsBoundToObject(this))
-				{
-					SkelMeshComp->OnTickPose.AddUObject(this, &UContextualAnimSceneActorComponent::OnTickPose);
-				}
+				SkelMeshComp->OnTickPose.AddUObject(this, &UContextualAnimSceneActorComponent::OnTickPose);
 			}
 
+			SetIgnoreCollisionWithOtherActors(true);
+
 			OnJoinedSceneDelegate.Broadcast(this);
 		}
 	}
@@ -65,7 +66,11 @@ void UContextualAnimSceneActorComponent::OnRep_Bindings()
 			SkelMeshComp->OnTickPose.RemoveAll(this);
 		}
 
+		SetIgnoreCollisionWithOtherActors(false);
+
 		OnLeftSceneDelegate.Broadcast(this);
+
+		LocalBindings.Reset();
 	}
 }
 
@@ -100,32 +105,63 @@ void UContextualAnimSceneActorComponent::OnUnregister()
 	}
 }
 
+void UContextualAnimSceneActorComponent::SetIgnoreCollisionWithOtherActors(bool bValue) const
+{
+	const AActor* OwnerActor = GetOwner();
+
+	for (const FContextualAnimSceneBinding& Binding : LocalBindings)
+	{
+		AActor* OtherActor = Binding.GetActor();
+		if (OtherActor != OwnerActor)
+		{
+			if (UPrimitiveComponent* RootPrimitiveComponent = Cast<UPrimitiveComponent>(OwnerActor->GetRootComponent()))
+			{
+				RootPrimitiveComponent->IgnoreActorWhenMoving(OtherActor, bValue);
+			}
+		}
+	}
+}
+
 void UContextualAnimSceneActorComponent::OnJoinedScene(const FContextualAnimSceneBindings& InBindings)
 {
+	if(LocalBindings.IsValid())
+	{
+		OnLeftScene();
+	}
+
 	if (const FContextualAnimSceneBinding* Binding = InBindings.FindBindingByActor(GetOwner()))
 	{
-		Bindings = InBindings;
+		LocalBindings = InBindings;
 
-		// Start listening to TickPose if we joined an scene where we need IK
-		if (Bindings.GetIKTargetDefContainerFromBinding(*Binding).IKTargetDefs.Num() > 0)
+		USkeletalMeshComponent* SkelMeshComp = UContextualAnimUtilities::TryGetSkeletalMeshComponent(GetOwner());
+		if (SkelMeshComp && !SkelMeshComp->OnTickPose.IsBoundToObject(this))
 		{
-			USkeletalMeshComponent* SkelMeshComp = UContextualAnimUtilities::TryGetSkeletalMeshComponent(GetOwner());
-			if (SkelMeshComp && !SkelMeshComp->OnTickPose.IsBoundToObject(this))
-			{
-				SkelMeshComp->OnTickPose.AddUObject(this, &UContextualAnimSceneActorComponent::OnTickPose);
-			}
+			SkelMeshComp->OnTickPose.AddUObject(this, &UContextualAnimSceneActorComponent::OnTickPose);
+		}
+
+		// Disable collision between actors so they can align perfectly
+		SetIgnoreCollisionWithOtherActors(true);
+
+		// Prevent physics rotation. During the interaction we want to be fully root motion driven
+		if (UCharacterMovementComponent* MovementComp = GetOwner()->FindComponentByClass<UCharacterMovementComponent>())
+		{
+			bAllowPhysicsRotationDuringAnimRootMotionBackup = MovementComp->bAllowPhysicsRotationDuringAnimRootMotion;
+			MovementComp->bAllowPhysicsRotationDuringAnimRootMotion = false;
 		}
 
 		OnJoinedSceneDelegate.Broadcast(this);
+
+		if(GetOwner()->HasAuthority())
+		{
+			RepBindings = LocalBindings;
+		}
 	}
 }
 
 void UContextualAnimSceneActorComponent::OnLeftScene()
 {
-	if (const FContextualAnimSceneBinding* Binding = Bindings.FindBindingByActor(GetOwner()))
+	if (const FContextualAnimSceneBinding* Binding = LocalBindings.FindBindingByActor(GetOwner()))
 	{
-		OnLeftSceneDelegate.Broadcast(this);
-
 		// Stop listening to TickPose if we were
 		USkeletalMeshComponent* SkelMeshComp = UContextualAnimUtilities::TryGetSkeletalMeshComponent(GetOwner());
 		if (SkelMeshComp && SkelMeshComp->OnTickPose.IsBoundToObject(this))
@@ -133,21 +169,62 @@ void UContextualAnimSceneActorComponent::OnLeftScene()
 			SkelMeshComp->OnTickPose.RemoveAll(this);
 		}
 
-		Bindings.Reset();
+		// Restore collision between actors
+		// Note that this assumes that we are the only one disabling the collision between these actors. 
+		// We might want to add a more robust mechanism to avoid overriding a request to disable collision that may have been set by another system
+		SetIgnoreCollisionWithOtherActors(false);
+
+		// Restore bAllowPhysicsRotationDuringAnimRootMotion
+		if (UCharacterMovementComponent* MovementComp = GetOwner()->FindComponentByClass<UCharacterMovementComponent>())
+		{
+			MovementComp->bAllowPhysicsRotationDuringAnimRootMotion = bAllowPhysicsRotationDuringAnimRootMotionBackup;
+		}
+
+		OnLeftSceneDelegate.Broadcast(this);
+
+		LocalBindings.Reset();
+
+		if (GetOwner()->HasAuthority())
+		{
+			RepBindings = LocalBindings;
+		}
 	}
 }
 
 void UContextualAnimSceneActorComponent::OnTickPose(class USkinnedMeshComponent* SkinnedMeshComponent, float DeltaTime, bool bNeedsValidRootMotion)
 {
-	//@TODO: Check for LOD to prevent this update if the actor is too far away
-	UpdateIKTargets();
+	if (const FContextualAnimSceneBinding* Binding = LocalBindings.FindBindingByActor(GetOwner()))
+	{
+		// Synchronize playback time with the leader
+		FAnimMontageInstance* MontageInstance = Binding->GetAnimMontageInstance();
+		if (MontageInstance && MontageInstance->GetMontageSyncLeader() == nullptr)
+		{
+			if (const FContextualAnimSceneBinding* SyncLeader = LocalBindings.GetSyncLeader())
+			{
+				if (SyncLeader->GetActor() != GetOwner())
+				{
+					if (FAnimMontageInstance* LeaderMontageInstance = SyncLeader->GetAnimMontageInstance())
+					{
+						if (LeaderMontageInstance->Montage == LocalBindings.GetAnimTrackFromBinding(*SyncLeader).Animation &&
+							MontageInstance->Montage == LocalBindings.GetAnimTrackFromBinding(*Binding).Animation)
+						{
+							MontageInstance->MontageSync_Follow(LeaderMontageInstance);
+						}
+					}
+				}
+			}
+		}
+
+		//@TODO: Check for LOD to prevent this update if the actor is too far away
+		UpdateIKTargets();
+	}
 }
 
 void UContextualAnimSceneActorComponent::UpdateIKTargets()
 {
 	IKTargets.Reset();
 
-	const FContextualAnimSceneBinding* BindingPtr = Bindings.FindBindingByActor(GetOwner());
+	const FContextualAnimSceneBinding* BindingPtr = LocalBindings.FindBindingByActor(GetOwner());
 	if (BindingPtr == nullptr)
 	{
 		return;
@@ -159,7 +236,7 @@ void UContextualAnimSceneActorComponent::UpdateIKTargets()
 		return;
 	}
 
-	const TArray<FContextualAnimIKTargetDefinition>& IKTargetDefs = Bindings.GetIKTargetDefContainerFromBinding(*BindingPtr).IKTargetDefs;
+	const TArray<FContextualAnimIKTargetDefinition>& IKTargetDefs = LocalBindings.GetIKTargetDefContainerFromBinding(*BindingPtr).IKTargetDefs;
 	for (const FContextualAnimIKTargetDefinition& IKTargetDef : IKTargetDefs)
 	{
 #if !(UE_BUILD_SHIPPING || UE_BUILD_TEST)
@@ -175,7 +252,7 @@ void UContextualAnimSceneActorComponent::UpdateIKTargets()
 		// @TODO: IKTargetTransform will be off by 1 frame if we tick before target. 
 		// Should we at least add an option to the SceneAsset to setup tick dependencies or should this be entirely up to the user?
 
-		if (const FContextualAnimSceneBinding* TargetBinding = Bindings.FindBindingByRole(IKTargetDef.TargetRoleName))
+		if (const FContextualAnimSceneBinding* TargetBinding = LocalBindings.FindBindingByRole(IKTargetDef.TargetRoleName))
 		{
 			if (const USkeletalMeshComponent* TargetSkelMeshComp = TargetBinding->GetSkeletalMeshComponent())
 			{
@@ -184,7 +261,7 @@ void UContextualAnimSceneActorComponent::UpdateIKTargets()
 					const FTransform IKTargetParentTransform = TargetSkelMeshComp->GetSocketTransform(IKTargetDef.TargetBoneName);
 
 					const float Time = MontageInstance->GetPosition();
-					IKTargetTransform = Bindings.GetIKTargetTransformFromBinding(*BindingPtr, IKTargetDef.GoalName, Time) * IKTargetParentTransform;
+					IKTargetTransform = LocalBindings.GetIKTargetTransformFromBinding(*BindingPtr, IKTargetDef.GoalName, Time) * IKTargetParentTransform;
 
 #if !(UE_BUILD_SHIPPING || UE_BUILD_TEST)
 					if (bDrawDebugEnable)
diff --git a/Engine/Plugins/Experimental/Animation/ContextualAnimation/Source/ContextualAnimation/Private/ContextualAnimTypes.cpp b/Engine/Plugins/Experimental/Animation/ContextualAnimation/Source/ContextualAnimation/Private/ContextualAnimTypes.cpp
--- a/Engine/Plugins/Experimental/Animation/ContextualAnimation/Source/ContextualAnimation/Private/ContextualAnimTypes.cpp
+++ b/Engine/Plugins/Experimental/Animation/ContextualAnimation/Source/ContextualAnimation/Private/ContextualAnimTypes.cpp
@@ -353,15 +353,20 @@ FContextualAnimSceneBindings::FContextualAnimSceneBindings(const UContextualAnim
 	SceneAsset = &InSceneAsset;
 	SectionIdx = InSectionIdx;
 	AnimSetIdx = InAnimSetIdx;
+
+	static uint8 IncrementID = 0;
+	IncrementID = IncrementID < UINT8_MAX ? IncrementID + 1 : 0;
+	Id = IncrementID;
 }
 
 bool FContextualAnimSceneBindings::IsValid() const
 {
-	return SceneAsset.IsValid() && SceneAsset->HasValidData() && Num() > 0;
+	return Id != 0 && SceneAsset.IsValid() && SceneAsset->HasValidData() && Num() > 0;
 }
 
 void FContextualAnimSceneBindings::Reset()
 {
+	Id = 0;
 	SceneAsset.Reset();
 	SectionIdx = INDEX_NONE;
 	AnimSetIdx = INDEX_NONE;
@@ -369,6 +374,12 @@ void FContextualAnimSceneBindings::Reset()
 	SceneInstancePtr.Reset();
 }
 
+const FContextualAnimSceneBinding* FContextualAnimSceneBindings::GetSyncLeader() const
+{
+	//@TODO: Return first secondary binding as sync leader for now. This may have to be explicitly defined, either in the SceneAsset or when creating the bindings.
+	return Data.FindByPredicate([this](const FContextualAnimSceneBinding& Item) { return GetRoleFromBinding(Item) != SceneAsset->GetPrimaryRole(); });
+}
+
 const FContextualAnimTrack& FContextualAnimSceneBindings::GetAnimTrackFromBinding(const FContextualAnimSceneBinding& Binding) const
 {
 	if (const FContextualAnimTrack* AnimTrack = GetSceneAsset()->GetAnimTrack(SectionIdx, AnimSetIdx, Binding.AnimTrackIdx))
@@ -401,6 +412,7 @@ const FName& FContextualAnimSceneBindings::GetRoleFromBinding(const FContextualA
 
 bool FContextualAnimSceneBindings::NetSerialize(FArchive& Ar, class UPackageMap* Map, bool& bOutSuccess)
 {
+	Ar << Id;
 	Ar << SceneAsset;
 	Ar << SectionIdx;
 	Ar << AnimSetIdx;
diff --git a/Engine/Plugins/Experimental/Animation/ContextualAnimation/Source/ContextualAnimation/Public/ContextualAnimSceneActorComponent.h b/Engine/Plugins/Experimental/Animation/ContextualAnimation/Source/ContextualAnimation/Public/ContextualAnimSceneActorComponent.h
--- a/Engine/Plugins/Experimental/Animation/ContextualAnimation/Source/ContextualAnimation/Public/ContextualAnimSceneActorComponent.h
+++ b/Engine/Plugins/Experimental/Animation/ContextualAnimation/Source/ContextualAnimation/Public/ContextualAnimSceneActorComponent.h
@@ -62,13 +62,29 @@ public:
 
 protected:
 
+	/** 
+	 * Replicated version of the bindings for the interaction we are currently playing.
+	 * @TODO: Right now we are replicating the same set of bindings from each actor in the interaction. 
+	 * This is not a big structure so it might be ok, but as an optimization we could replicate it only from the leader of the interaction and set it on the other actors in the OnRep notify
+	 */
 	UPROPERTY(Transient, ReplicatedUsing = OnRep_Bindings)
-	FContextualAnimSceneBindings Bindings;
+	FContextualAnimSceneBindings RepBindings;
+
+	/** 
+	 * Local copy of the bindings for the interaction we are currently playing.
+	 * Used to update IK, keep montage in sync and disable/enable collision between actors on simulated proxies too 
+	 */
+	UPROPERTY(Transient)
+	FContextualAnimSceneBindings LocalBindings;
 
 	/** List of IKTarget for this frame */
-	UPROPERTY()
+	UPROPERTY(Transient)
 	TArray<FContextualAnimIKTarget> IKTargets;
 
+	/** Value of AllowPhysicsRotationDuringAnimRootMotion property for this actor before we join an scene, so we can restore it once the interaction ends */
+	UPROPERTY(Transient)
+	bool bAllowPhysicsRotationDuringAnimRootMotionBackup = false;
+
 	void UpdateIKTargets();
 
 	/** 
@@ -81,6 +97,8 @@ protected:
 	UFUNCTION()
 	void OnRep_Bindings();
 
+	void SetIgnoreCollisionWithOtherActors(bool bValue) const;
+
 private:
 
 	bool bRegistered = false;
diff --git a/Engine/Plugins/Experimental/Animation/ContextualAnimation/Source/ContextualAnimation/Public/ContextualAnimTypes.h b/Engine/Plugins/Experimental/Animation/ContextualAnimation/Source/ContextualAnimation/Public/ContextualAnimTypes.h
--- a/Engine/Plugins/Experimental/Animation/ContextualAnimation/Source/ContextualAnimation/Public/ContextualAnimTypes.h
+++ b/Engine/Plugins/Experimental/Animation/ContextualAnimation/Source/ContextualAnimation/Public/ContextualAnimTypes.h
@@ -394,6 +394,7 @@ struct CONTEXTUALANIMATION_API FContextualAnimSceneBindings
 		return Role != NAME_None ? Data.FindByPredicate([this, &Role](const FContextualAnimSceneBinding& Item) { return GetAnimTrackFromBinding(Item).Role == Role; }) : nullptr;
 	}
 
+	FORCEINLINE uint8 GetID() const { return Id; }
 	FORCEINLINE const UContextualAnimSceneAsset* GetSceneAsset() const { return SceneAsset.Get(); }
 	FORCEINLINE int32 GetSectionIdx() const { return SectionIdx; }
 	FORCEINLINE int32 GetAnimSetIdx() const { return AnimSetIdx; }
@@ -421,6 +422,8 @@ struct CONTEXTUALANIMATION_API FContextualAnimSceneBindings
 	const FContextualAnimIKTargetDefContainer& GetIKTargetDefContainerFromBinding(const FContextualAnimSceneBinding& Binding) const;
 	FTransform GetIKTargetTransformFromBinding(const FContextualAnimSceneBinding& Binding, const FName& TrackName, float Time) const;
 
+	const FContextualAnimSceneBinding* GetSyncLeader() const;
+
 	bool IsValid() const;
 
 	void Reset();
@@ -431,6 +434,9 @@ private:
 
 	friend class UContextualAnimManager;
 
+	UPROPERTY()
+	uint8 Id = 0;
+
 	UPROPERTY()
 	TWeakObjectPtr<const UContextualAnimSceneAsset> SceneAsset = nullptr;
 
