Contextual Anim: Refactored the way bindings are replicated to simulated proxies. SceneActorComp now also replicates animation playback to simulated proxies. Ensure SceneActorComp exist on every actor during the binding process.

diff --git a/Engine/Plugins/Experimental/Animation/ContextualAnimation/Source/ContextualAnimation/ContextualAnimation.Build.cs b/Engine/Plugins/Experimental/Animation/ContextualAnimation/Source/ContextualAnimation/ContextualAnimation.Build.cs
--- a/Engine/Plugins/Experimental/Animation/ContextualAnimation/Source/ContextualAnimation/ContextualAnimation.Build.cs
+++ b/Engine/Plugins/Experimental/Animation/ContextualAnimation/Source/ContextualAnimation/ContextualAnimation.Build.cs
@@ -14,6 +14,7 @@ public class ContextualAnimation : ModuleRules
 				"Core",
 				"CoreUObject",
                 "Engine",
+				"NetCore",
 				"GameplayTags",
 				"MotionWarping",
 				"IKRig"
diff --git a/Engine/Plugins/Experimental/Animation/ContextualAnimation/Source/ContextualAnimation/Private/ContextualAnimSceneActorComponent.cpp b/Engine/Plugins/Experimental/Animation/ContextualAnimation/Source/ContextualAnimation/Private/ContextualAnimSceneActorComponent.cpp
--- a/Engine/Plugins/Experimental/Animation/ContextualAnimation/Source/ContextualAnimation/Private/ContextualAnimSceneActorComponent.cpp
+++ b/Engine/Plugins/Experimental/Animation/ContextualAnimation/Source/ContextualAnimation/Private/ContextualAnimSceneActorComponent.cpp
@@ -14,6 +14,7 @@
 #include "AnimNotifyState_IKWindow.h"
 #include "EngineUtils.h"
 #include "Net/UnrealNetwork.h"
+#include "Net/Core/PushModel/PushModel.h"
 #include "PrimitiveSceneProxy.h"
 #include "SceneManagement.h"
 
@@ -36,41 +37,88 @@ UContextualAnimSceneActorComponent::UContextualAnimSceneActorComponent(const FOb
 void UContextualAnimSceneActorComponent::GetLifetimeReplicatedProps(TArray< FLifetimeProperty >& OutLifetimeProps) const
 {
 	Super::GetLifetimeReplicatedProps(OutLifetimeProps);
-	DOREPLIFETIME_CONDITION(UContextualAnimSceneActorComponent, RepBindings, COND_SimulatedOnly);
+
+	FDoRepLifetimeParams Params;
+	Params.bIsPushBased = true;
+	Params.Condition = COND_SimulatedOnly;
+	DOREPLIFETIME_WITH_PARAMS_FAST(UContextualAnimSceneActorComponent, RepBindings, Params);
 }
 
-void UContextualAnimSceneActorComponent::OnRep_Bindings()
+bool UContextualAnimSceneActorComponent::StartContextualAnimScene(const FContextualAnimSceneBindings& InBindings)
 {
-	if (RepBindings.IsValid())
+	UE_LOG(LogContextualAnim, Log, TEXT("%-21s UContextualAnimSceneActorComponent::StartContextualAnim Actor: %s"), *UEnum::GetValueAsString(TEXT("Engine.ENetRole"), GetOwner()->GetLocalRole()), *GetNameSafe(GetOwner()));
+
+	const FContextualAnimSceneBinding* OwnerBinding = InBindings.FindBindingByActor(GetOwner());
+	if (ensureAlways(OwnerBinding))
 	{
-		if (const FContextualAnimSceneBinding* Binding = RepBindings.FindBindingByActor(GetOwner()))
-		{
-			LocalBindings = RepBindings;
+		OwnerBinding->GetSceneActorComponent()->JoinScene(InBindings);
 
-			USkeletalMeshComponent* SkelMeshComp = UContextualAnimUtilities::TryGetSkeletalMeshComponent(GetOwner());
-			if (SkelMeshComp && !SkelMeshComp->OnTickPose.IsBoundToObject(this))
+		for (const FContextualAnimSceneBinding& Binding : InBindings)
+		{
+			if (Binding.GetActor() != GetOwner())
 			{
-				SkelMeshComp->OnTickPose.AddUObject(this, &UContextualAnimSceneActorComponent::OnTickPose);
+				Binding.GetSceneActorComponent()->JoinScene(InBindings);
 			}
+		}
 
-			SetIgnoreCollisionWithOtherActors(true);
+		//@TODO: Temp until we move the scene pivots to the bindings
+		UContextualAnimUtilities::BP_SceneBindings_AddOrUpdateWarpTargetsForBindings(InBindings);
 
-			OnJoinedSceneDelegate.Broadcast(this);
+		if (GetOwner()->HasAuthority())
+		{
+			RepBindings = InBindings;
+			MARK_PROPERTY_DIRTY_FROM_NAME(UContextualAnimSceneActorComponent, RepBindings, this);
 		}
+
+		return true;
 	}
-	else
+
+	return false;
+}
+
+void UContextualAnimSceneActorComponent::OnRep_Bindings(const FContextualAnimSceneBindings& LastRepBindings)
+{
+	// @TODO: This need more investigation but for now it prevents an issue caused by this OnRep_ triggering even when there is no (obvious) changein the data
+	if(RepBindings.GetID() == LastRepBindings.GetID())
 	{
-		USkeletalMeshComponent* SkelMeshComp = UContextualAnimUtilities::TryGetSkeletalMeshComponent(GetOwner());
-		if (SkelMeshComp && SkelMeshComp->OnTickPose.IsBoundToObject(this))
-		{
-			SkelMeshComp->OnTickPose.RemoveAll(this);
-		}
+		UE_LOG(LogContextualAnim, Warning, TEXT("%-21s UContextualAnimSceneActorComponent::OnRep_Bindings Actor: %s RepBindings Id: %d LastRepBindings Id: %d"),
+			*UEnum::GetValueAsString(TEXT("Engine.ENetRole"), GetOwner()->GetLocalRole()), *GetNameSafe(GetOwner()), RepBindings.GetID(), LastRepBindings.GetID());
 
-		SetIgnoreCollisionWithOtherActors(false);
+		return;
+	}
 
-		OnLeftSceneDelegate.Broadcast(this);
+	UE_LOG(LogContextualAnim, Verbose, TEXT("%-21s UContextualAnimSceneActorComponent::OnRep_Bindings Actor: %s RepBindings Id: %d Bindings Id: %d"),
+		*UEnum::GetValueAsString(TEXT("Engine.ENetRole"), GetOwner()->GetLocalRole()), *GetNameSafe(GetOwner()), RepBindings.GetID(), Bindings.GetID());
 
-		LocalBindings.Reset();
+	// The owner of this component started an interaction on the server
+	if (RepBindings.IsValid())
+	{
+		const FContextualAnimSceneBinding* OwnerBinding = RepBindings.FindBindingByActor(GetOwner());
+		if (ensureAlways(OwnerBinding))
+		{
+			// Join the scene (start playing animation, etc.)
+			OwnerBinding->GetSceneActorComponent()->JoinScene(RepBindings);
+
+			// RepBindings is only replicated from the initiator of the actor. 
+			// So now we have to tell everyone else involved in the interaction to join us
+			// @TODO: For now this assumes that all the actors will start playing the animation at the same time. 
+			// We will expand this in the future to allow 'late' join
+			for (const FContextualAnimSceneBinding& Binding : RepBindings)
+			{
+				if (Binding.GetActor() != GetOwner())
+				{
+					Binding.GetSceneActorComponent()->JoinScene(RepBindings);
+				}
+			}
+		}
+	}
+	else
+	{	
+		// Empty bindings is replicated by the initiator of the interaction when the animation ends
+		// In this case we don't want to tell everyone else to also leave the scene since there is very common for the initiator, 
+		// specially if is player character, to end the animation earlier for responsiveness
+		// It is more likely this will do nothing since we listen to montage end also on Simulated Proxies to 'predict' the end of the interaction.
+		LeaveScene();
 	}
 }
 
@@ -109,7 +157,7 @@ void UContextualAnimSceneActorComponent::SetIgnoreCollisionWithOtherActors(bool
 {
 	const AActor* OwnerActor = GetOwner();
 
-	for (const FContextualAnimSceneBinding& Binding : LocalBindings)
+	for (const FContextualAnimSceneBinding& Binding : Bindings)
 	{
 		AActor* OtherActor = Binding.GetActor();
 		if (OtherActor != OwnerActor)
@@ -124,14 +172,17 @@ void UContextualAnimSceneActorComponent::SetIgnoreCollisionWithOtherActors(bool
 
 void UContextualAnimSceneActorComponent::OnJoinedScene(const FContextualAnimSceneBindings& InBindings)
 {
-	if(LocalBindings.IsValid())
+	UE_LOG(LogContextualAnim, Verbose, TEXT("%-21s UContextualAnimSceneActorComponent::OnJoinedScene Actor: %s InBindings Id: %d"),
+		*UEnum::GetValueAsString(TEXT("Engine.ENetRole"), GetOwner()->GetLocalRole()), *GetNameSafe(GetOwner()), InBindings.GetID());
+
+	if (Bindings.IsValid())
 	{
 		OnLeftScene();
 	}
 
 	if (const FContextualAnimSceneBinding* Binding = InBindings.FindBindingByActor(GetOwner()))
 	{
-		LocalBindings = InBindings;
+		Bindings = InBindings;
 
 		USkeletalMeshComponent* SkelMeshComp = UContextualAnimUtilities::TryGetSkeletalMeshComponent(GetOwner());
 		if (SkelMeshComp && !SkelMeshComp->OnTickPose.IsBoundToObject(this))
@@ -150,20 +201,115 @@ void UContextualAnimSceneActorComponent::OnJoinedScene(const FContextualAnimScen
 		}
 
 		OnJoinedSceneDelegate.Broadcast(this);
+	}
+}
+
+void UContextualAnimSceneActorComponent::OnLeftScene()
+{
+	UE_LOG(LogContextualAnim, Verbose, TEXT("%-21s UContextualAnimSceneActorComponent::OnLeftScene Actor: %s Current Bindings Id: %d"),
+		*UEnum::GetValueAsString(TEXT("Engine.ENetRole"), GetOwner()->GetLocalRole()), *GetNameSafe(GetOwner()), Bindings.GetID());
+
+	if (const FContextualAnimSceneBinding* Binding = Bindings.FindBindingByActor(GetOwner()))
+	{
+		// Stop listening to TickPose if we were
+		USkeletalMeshComponent* SkelMeshComp = UContextualAnimUtilities::TryGetSkeletalMeshComponent(GetOwner());
+		if (SkelMeshComp && SkelMeshComp->OnTickPose.IsBoundToObject(this))
+		{
+			SkelMeshComp->OnTickPose.RemoveAll(this);
+		}
+
+		// Restore collision between actors
+		// Note that this assumes that we are the only one disabling the collision between these actors. 
+		// We might want to add a more robust mechanism to avoid overriding a request to disable collision that may have been set by another system
+		SetIgnoreCollisionWithOtherActors(false);
 
-		if(GetOwner()->HasAuthority())
+		// Restore bAllowPhysicsRotationDuringAnimRootMotion
+		if (UCharacterMovementComponent* MovementComp = GetOwner()->FindComponentByClass<UCharacterMovementComponent>())
 		{
-			RepBindings = LocalBindings;
+			MovementComp->bAllowPhysicsRotationDuringAnimRootMotion = bAllowPhysicsRotationDuringAnimRootMotionBackup;
 		}
+
+		OnLeftSceneDelegate.Broadcast(this);
+
+		Bindings.Reset();
 	}
 }
 
-void UContextualAnimSceneActorComponent::OnLeftScene()
+void UContextualAnimSceneActorComponent::JoinScene(const FContextualAnimSceneBindings& InBindings)
 {
-	if (const FContextualAnimSceneBinding* Binding = LocalBindings.FindBindingByActor(GetOwner()))
+	if (Bindings.IsValid())
 	{
+		LeaveScene();
+	}
+
+	if (const FContextualAnimSceneBinding* Binding = InBindings.FindBindingByActor(GetOwner()))
+	{
+		UE_LOG(LogContextualAnim, Verbose, TEXT("%-21s UContextualAnimSceneActorComponent::JoinScene Actor: %s InBindings Id: %d Section: %d Asset: %s"),
+			*UEnum::GetValueAsString(TEXT("Engine.ENetRole"), GetOwner()->GetLocalRole()), *GetNameSafe(GetOwner()), InBindings.GetID(), InBindings.GetSectionIdx(), *GetNameSafe(InBindings.GetSceneAsset()));
+
+		Bindings = InBindings;
+
+		const FContextualAnimTrack& AnimTrack = Bindings.GetAnimTrackFromBinding(*Binding);
+
+		if (UAnimInstance* AnimInstance = Binding->GetAnimInstance())
+		{
+			AnimInstance->OnMontageBlendingOut.AddUniqueDynamic(this, &UContextualAnimSceneActorComponent::OnMontageBlendingOut);
+
+			//@TODO: Add support for dynamic montage
+			UAnimMontage* AnimMontage = Cast<UAnimMontage>(AnimTrack.Animation);
+			AnimInstance->Montage_Play(AnimMontage, 1.f);
+		}
+
+		USkeletalMeshComponent* SkelMeshComp = Binding->GetSkeletalMeshComponent();
+		if (SkelMeshComp && !SkelMeshComp->OnTickPose.IsBoundToObject(this))
+		{
+			SkelMeshComp->OnTickPose.AddUObject(this, &UContextualAnimSceneActorComponent::OnTickPose);
+		}
+
+		// Disable collision between actors so they can align perfectly
+		SetIgnoreCollisionWithOtherActors(true);
+
+		// Prevent physics rotation. During the interaction we want to be fully root motion driven
+		if (UCharacterMovementComponent* MovementComp = GetOwner()->FindComponentByClass<UCharacterMovementComponent>())
+		{
+			bAllowPhysicsRotationDuringAnimRootMotionBackup = MovementComp->bAllowPhysicsRotationDuringAnimRootMotion;
+			MovementComp->bAllowPhysicsRotationDuringAnimRootMotion = false;
+
+			//@TODO: Temp solution that assumes these interactions are not locally predicted and that is ok to be in flying mode during the entire animation
+			if (AnimTrack.bRequireFlyingMode && MovementComp->MovementMode != MOVE_Flying)
+			{
+				MovementComp->SetMovementMode(MOVE_Flying);
+			}
+		}
+
+		OnJoinedSceneDelegate.Broadcast(this);
+	}
+}
+
+void UContextualAnimSceneActorComponent::LeaveScene()
+{
+	if (const FContextualAnimSceneBinding* Binding = Bindings.FindBindingByActor(GetOwner()))
+	{
+		UE_LOG(LogContextualAnim, Verbose, TEXT("%-21s UContextualAnimSceneActorComponent::LeaveScene Actor: %s Current Bindings Id: %d Section: %d Asset: %s"),
+			*UEnum::GetValueAsString(TEXT("Engine.ENetRole"), GetOwner()->GetLocalRole()), *GetNameSafe(GetOwner()), Bindings.GetID(), Bindings.GetSectionIdx(), *GetNameSafe(Bindings.GetSceneAsset()));
+
+		const FContextualAnimTrack& AnimTrack = Bindings.GetAnimTrackFromBinding(*Binding);
+
+		if (UAnimInstance* AnimInstance = Binding->GetAnimInstance())
+		{
+			AnimInstance->OnMontageBlendingOut.RemoveDynamic(this, &UContextualAnimSceneActorComponent::OnMontageBlendingOut);
+
+			//@TODO: Add support for dynamic montage
+			UAnimMontage* AnimMontage = Cast<UAnimMontage>(AnimTrack.Animation);
+
+			if (AnimInstance->Montage_IsPlaying(AnimMontage))
+			{
+				AnimInstance->Montage_Stop(AnimMontage->GetDefaultBlendOutTime(), AnimMontage);
+			}
+		}
+
 		// Stop listening to TickPose if we were
-		USkeletalMeshComponent* SkelMeshComp = UContextualAnimUtilities::TryGetSkeletalMeshComponent(GetOwner());
+		USkeletalMeshComponent* SkelMeshComp = Binding->GetSkeletalMeshComponent();
 		if (SkelMeshComp && SkelMeshComp->OnTickPose.IsBoundToObject(this))
 		{
 			SkelMeshComp->OnTickPose.RemoveAll(this);
@@ -178,35 +324,62 @@ void UContextualAnimSceneActorComponent::OnLeftScene()
 		if (UCharacterMovementComponent* MovementComp = GetOwner()->FindComponentByClass<UCharacterMovementComponent>())
 		{
 			MovementComp->bAllowPhysicsRotationDuringAnimRootMotion = bAllowPhysicsRotationDuringAnimRootMotionBackup;
+
+			//@TODO: Temp solution that assumes these interactions are not locally predicted and that is ok to be in flying mode during the entire animation
+			if (AnimTrack.bRequireFlyingMode && MovementComp->MovementMode == MOVE_Flying)
+			{
+				MovementComp->SetMovementMode(MOVE_Walking);
+			}
 		}
 
 		OnLeftSceneDelegate.Broadcast(this);
 
-		LocalBindings.Reset();
+		Bindings.Reset();
+	}
+}
 
-		if (GetOwner()->HasAuthority())
+void UContextualAnimSceneActorComponent::OnMontageBlendingOut(UAnimMontage* Montage, bool bInterrupted)
+{
+	UE_LOG(LogContextualAnim, Verbose, TEXT("%-21s UContextualAnimSceneActorComponent::OnMontageBlendingOut Actor: %s Montage: %s bInterrupted: %d"),
+		*UEnum::GetValueAsString(TEXT("Engine.ENetRole"), GetOwner()->GetLocalRole()), *GetNameSafe(GetOwner()), *GetNameSafe(Montage), bInterrupted);
+
+	if (const FContextualAnimSceneBinding* Binding = Bindings.FindBindingByActor(GetOwner()))
+	{
+		if (Bindings.GetAnimTrackFromBinding(*Binding).Animation == Montage)
 		{
-			RepBindings = LocalBindings;
+			LeaveScene();
+
+			if (GetOwner()->HasAuthority())
+			{
+				// Rep empty bindings if we were the initiator of this interaction.
+				if(RepBindings.IsValid())
+				{
+					RepBindings.Reset();
+					MARK_PROPERTY_DIRTY_FROM_NAME(UContextualAnimSceneActorComponent, RepBindings, this);
+				}
+
+				//@TODO: Replicate this event separately for each other member of the interaction
+			}
 		}
 	}
 }
 
 void UContextualAnimSceneActorComponent::OnTickPose(class USkinnedMeshComponent* SkinnedMeshComponent, float DeltaTime, bool bNeedsValidRootMotion)
 {
-	if (const FContextualAnimSceneBinding* Binding = LocalBindings.FindBindingByActor(GetOwner()))
+	if (const FContextualAnimSceneBinding* Binding = Bindings.FindBindingByActor(GetOwner()))
 	{
 		// Synchronize playback time with the leader
 		FAnimMontageInstance* MontageInstance = Binding->GetAnimMontageInstance();
 		if (MontageInstance && MontageInstance->GetMontageSyncLeader() == nullptr)
 		{
-			if (const FContextualAnimSceneBinding* SyncLeader = LocalBindings.GetSyncLeader())
+			if (const FContextualAnimSceneBinding* SyncLeader = Bindings.GetSyncLeader())
 			{
 				if (SyncLeader->GetActor() != GetOwner())
 				{
 					if (FAnimMontageInstance* LeaderMontageInstance = SyncLeader->GetAnimMontageInstance())
 					{
-						if (LeaderMontageInstance->Montage == LocalBindings.GetAnimTrackFromBinding(*SyncLeader).Animation &&
-							MontageInstance->Montage == LocalBindings.GetAnimTrackFromBinding(*Binding).Animation)
+						if (LeaderMontageInstance->Montage == Bindings.GetAnimTrackFromBinding(*SyncLeader).Animation &&
+							MontageInstance->Montage == Bindings.GetAnimTrackFromBinding(*Binding).Animation)
 						{
 							MontageInstance->MontageSync_Follow(LeaderMontageInstance);
 						}
@@ -224,7 +397,7 @@ void UContextualAnimSceneActorComponent::UpdateIKTargets()
 {
 	IKTargets.Reset();
 
-	const FContextualAnimSceneBinding* BindingPtr = LocalBindings.FindBindingByActor(GetOwner());
+	const FContextualAnimSceneBinding* BindingPtr = Bindings.FindBindingByActor(GetOwner());
 	if (BindingPtr == nullptr)
 	{
 		return;
@@ -236,7 +409,7 @@ void UContextualAnimSceneActorComponent::UpdateIKTargets()
 		return;
 	}
 
-	const TArray<FContextualAnimIKTargetDefinition>& IKTargetDefs = LocalBindings.GetIKTargetDefContainerFromBinding(*BindingPtr).IKTargetDefs;
+	const TArray<FContextualAnimIKTargetDefinition>& IKTargetDefs = Bindings.GetIKTargetDefContainerFromBinding(*BindingPtr).IKTargetDefs;
 	for (const FContextualAnimIKTargetDefinition& IKTargetDef : IKTargetDefs)
 	{
 #if !(UE_BUILD_SHIPPING || UE_BUILD_TEST)
@@ -252,7 +425,7 @@ void UContextualAnimSceneActorComponent::UpdateIKTargets()
 		// @TODO: IKTargetTransform will be off by 1 frame if we tick before target. 
 		// Should we at least add an option to the SceneAsset to setup tick dependencies or should this be entirely up to the user?
 
-		if (const FContextualAnimSceneBinding* TargetBinding = LocalBindings.FindBindingByRole(IKTargetDef.TargetRoleName))
+		if (const FContextualAnimSceneBinding* TargetBinding = Bindings.FindBindingByRole(IKTargetDef.TargetRoleName))
 		{
 			if (const USkeletalMeshComponent* TargetSkelMeshComp = TargetBinding->GetSkeletalMeshComponent())
 			{
@@ -261,7 +434,7 @@ void UContextualAnimSceneActorComponent::UpdateIKTargets()
 					const FTransform IKTargetParentTransform = TargetSkelMeshComp->GetSocketTransform(IKTargetDef.TargetBoneName);
 
 					const float Time = MontageInstance->GetPosition();
-					IKTargetTransform = LocalBindings.GetIKTargetTransformFromBinding(*BindingPtr, IKTargetDef.GoalName, Time) * IKTargetParentTransform;
+					IKTargetTransform = Bindings.GetIKTargetTransformFromBinding(*BindingPtr, IKTargetDef.GoalName, Time) * IKTargetParentTransform;
 
 #if !(UE_BUILD_SHIPPING || UE_BUILD_TEST)
 					if (bDrawDebugEnable)
diff --git a/Engine/Plugins/Experimental/Animation/ContextualAnimation/Source/ContextualAnimation/Private/ContextualAnimTypes.cpp b/Engine/Plugins/Experimental/Animation/ContextualAnimation/Source/ContextualAnimation/Private/ContextualAnimTypes.cpp
--- a/Engine/Plugins/Experimental/Animation/ContextualAnimation/Source/ContextualAnimation/Private/ContextualAnimTypes.cpp
+++ b/Engine/Plugins/Experimental/Animation/ContextualAnimation/Source/ContextualAnimation/Private/ContextualAnimTypes.cpp
@@ -288,19 +288,35 @@ void FContextualAnimSceneBinding::SetAnimTrack(const FContextualAnimTrack& InAni
 
 UContextualAnimSceneActorComponent* FContextualAnimSceneBinding::GetSceneActorComponent() const
 {
-	//@TODO: Cache this during the binding
-	AActor* Actor = Context.GetActor();
-	return Actor ? Actor->FindComponentByClass<UContextualAnimSceneActorComponent>() : nullptr;
+	if (CachedSceneActorComp == nullptr)
+	{
+		if (AActor* Actor = Context.GetActor())
+		{
+			CachedSceneActorComp = Actor->FindComponentByClass<UContextualAnimSceneActorComponent>();
+		}
+	}
+
+	return CachedSceneActorComp;
 }
 
 UAnimInstance* FContextualAnimSceneBinding::GetAnimInstance() const
 {
-	return UContextualAnimUtilities::TryGetAnimInstance(GetActor());
+	if (CachedAnimInstance == nullptr)
+	{
+		CachedAnimInstance = UContextualAnimUtilities::TryGetAnimInstance(GetActor());
+	}
+
+	return CachedAnimInstance;
 }
 
 USkeletalMeshComponent* FContextualAnimSceneBinding::GetSkeletalMeshComponent() const
 {
-	return UContextualAnimUtilities::TryGetSkeletalMeshComponent(GetActor());
+	if(CachedSkeletalMesh == nullptr)
+	{
+		CachedSkeletalMesh = UContextualAnimUtilities::TryGetSkeletalMeshComponent(GetActor());
+	}
+
+	return CachedSkeletalMesh;
 }
 
 FAnimMontageInstance* FContextualAnimSceneBinding::GetAnimMontageInstance() const
@@ -438,6 +454,19 @@ bool FContextualAnimSceneBindings::TryCreateBindings(const UContextualAnimSceneA
 		return false;
 	}
 
+	if(AActor* PrimaryActor = PrimaryPtr->GetActor())
+	{
+		UContextualAnimSceneActorComponent* SceneActorComp = PrimaryActor->FindComponentByClass<UContextualAnimSceneActorComponent>();
+		if(SceneActorComp == nullptr)
+		{
+			UE_LOG(LogContextualAnim, Warning, TEXT("FContextualAnimSceneBindings::TryCreateBindings Failed. Reason: Missing ContextualAnimSceneActorComp on primary actor. SceneAsset: %s Actor: %s"),
+				*GetNameSafe(&SceneAsset), *GetNameSafe(PrimaryActor));
+
+			OutBindings.Reset();
+			return false;
+		}
+	}
+
 	// First, try to bind primary track. 
 	// @TODO: Revisit this, passing the same data twice (as primary and querier) feels weird, but this allow us to run the selection mechanism even on the primary actor.
 
@@ -462,6 +491,19 @@ bool FContextualAnimSceneBindings::TryCreateBindings(const UContextualAnimSceneA
 		FName RoleToBind = Pair.Key;
 		if (RoleToBind != PrimaryRole)
 		{
+			if (AActor* SecondaryActor = Pair.Value.GetActor())
+			{
+				UContextualAnimSceneActorComponent* SceneActorComp = SecondaryActor->FindComponentByClass<UContextualAnimSceneActorComponent>();
+				if (SceneActorComp == nullptr)
+				{
+					UE_LOG(LogContextualAnim, Warning, TEXT("FContextualAnimSceneBindings::TryCreateBindings Failed. Reason: Missing ContextualAnimSceneActorComp on secondary actor. SceneAsset: %s Actor: %s"),
+						*GetNameSafe(&SceneAsset), *GetNameSafe(SecondaryActor));
+
+					OutBindings.Reset();
+					return false;
+				}
+			}
+
 			const FContextualAnimTrack* AnimTrack = SceneAsset.GetAnimTrack(SectionIdx, AnimSetIdx, RoleToBind);
 			if (AnimTrack && AnimTrack->DoesQuerierPassSelectionCriteria(*PrimaryPtr, Pair.Value))
 			{
@@ -512,6 +554,19 @@ bool FContextualAnimSceneBindings::TryCreateBindings(const UContextualAnimSceneA
 		const FName PrimaryRole = SceneAsset.GetPrimaryRole();
 		if(RoleDef.Name == PrimaryRole)
 		{
+			if (AActor* Actor = Primary.GetActor())
+			{
+				UContextualAnimSceneActorComponent* SceneActorComp = Actor->FindComponentByClass<UContextualAnimSceneActorComponent>();
+				if (SceneActorComp == nullptr)
+				{
+					UE_LOG(LogContextualAnim, Warning, TEXT("FContextualAnimSceneBindings::TryCreateBindings Failed. Reason: Missing ContextualAnimSceneActorComp on primary actor. SceneAsset: %s Actor: %s"),
+						*GetNameSafe(&SceneAsset), *GetNameSafe(Actor));
+
+					OutBindings.Reset();
+					return false;
+				}
+			}
+
 			const FContextualAnimTrack* AnimTrack = SceneAsset.GetAnimTrack(SectionIdx, AnimSetIdx, PrimaryRole);
 			if (AnimTrack && AnimTrack->DoesQuerierPassSelectionCriteria(Primary, Primary))
 			{
@@ -528,6 +583,19 @@ bool FContextualAnimSceneBindings::TryCreateBindings(const UContextualAnimSceneA
 		}
 		else // Secondary Role
 		{
+			if (AActor* Actor = Secondary.GetActor())
+			{
+				UContextualAnimSceneActorComponent* SceneActorComp = Actor->FindComponentByClass<UContextualAnimSceneActorComponent>();
+				if (SceneActorComp == nullptr)
+				{
+					UE_LOG(LogContextualAnim, Warning, TEXT("FContextualAnimSceneBindings::TryCreateBindings Failed. Reason: Missing ContextualAnimSceneActorComp on primary actor. SceneAsset: %s Actor: %s"),
+						*GetNameSafe(&SceneAsset), *GetNameSafe(Actor));
+
+					OutBindings.Reset();
+					return false;
+				}
+			}
+
 			const FContextualAnimTrack* AnimTrack = SceneAsset.GetAnimTrack(SectionIdx, AnimSetIdx, RoleDef.Name);
 			if (AnimTrack && AnimTrack->DoesQuerierPassSelectionCriteria(Primary, Secondary))
 			{
diff --git a/Engine/Plugins/Experimental/Animation/ContextualAnimation/Source/ContextualAnimation/Public/ContextualAnimSceneActorComponent.h b/Engine/Plugins/Experimental/Animation/ContextualAnimation/Source/ContextualAnimation/Public/ContextualAnimSceneActorComponent.h
--- a/Engine/Plugins/Experimental/Animation/ContextualAnimation/Source/ContextualAnimation/Public/ContextualAnimSceneActorComponent.h
+++ b/Engine/Plugins/Experimental/Animation/ContextualAnimation/Source/ContextualAnimation/Public/ContextualAnimSceneActorComponent.h
@@ -60,22 +60,24 @@ public:
 	UFUNCTION(BlueprintPure, Category = "Contextual Anim|Scene Actor Component")
 	const FContextualAnimIKTarget& GetIKTargetByGoalName(FName GoalName) const;
 
+	UFUNCTION(BlueprintCallable, Category = "Contextual Anim|Scene Actor Component")
+	bool StartContextualAnimScene(const FContextualAnimSceneBindings& InBindings);
+
 protected:
 
 	/** 
-	 * Replicated version of the bindings for the interaction we are currently playing.
-	 * @TODO: Right now we are replicating the same set of bindings from each actor in the interaction. 
-	 * This is not a big structure so it might be ok, but as an optimization we could replicate it only from the leader of the interaction and set it on the other actors in the OnRep notify
+	 * Replicated copy of the bindings so we can start the action on simulated proxies 
+	 * This gets replicated only from the initiator of the action and then set on all the other members of the interaction
 	 */
 	UPROPERTY(Transient, ReplicatedUsing = OnRep_Bindings)
 	FContextualAnimSceneBindings RepBindings;
 
-	/** 
-	 * Local copy of the bindings for the interaction we are currently playing.
-	 * Used to update IK, keep montage in sync and disable/enable collision between actors on simulated proxies too 
+	/**
+	 * Bindings for the interaction we are currently playing.
+	 * Used to update IK, keep montage in sync, disable/enable collision between actors etc
 	 */
 	UPROPERTY(Transient)
-	FContextualAnimSceneBindings LocalBindings;
+	FContextualAnimSceneBindings Bindings;
 
 	/** List of IKTarget for this frame */
 	UPROPERTY(Transient)
@@ -95,11 +97,22 @@ protected:
 	void OnTickPose(class USkinnedMeshComponent* SkinnedMeshComponent, float DeltaTime, bool bNeedsValidRootMotion);
 
 	UFUNCTION()
-	void OnRep_Bindings();
+	void OnRep_Bindings(const FContextualAnimSceneBindings& LastRepBindings);
 
 	void SetIgnoreCollisionWithOtherActors(bool bValue) const;
 
+	UFUNCTION()
+	void OnMontageBlendingOut(UAnimMontage* Montage, bool bInterrupted);
+
+	// @TODO: These two functions are going to replace OnJoinedScene and OnLeftScene
+	// main different is that these new functions are taking care of animation playback too
+
+	void JoinScene(const FContextualAnimSceneBindings& InBindings);
+
+	void LeaveScene();
+
 private:
 
 	bool bRegistered = false;
+
 };
diff --git a/Engine/Plugins/Experimental/Animation/ContextualAnimation/Source/ContextualAnimation/Public/ContextualAnimTypes.h b/Engine/Plugins/Experimental/Animation/ContextualAnimation/Source/ContextualAnimation/Public/ContextualAnimTypes.h
--- a/Engine/Plugins/Experimental/Animation/ContextualAnimation/Source/ContextualAnimation/Public/ContextualAnimTypes.h
+++ b/Engine/Plugins/Experimental/Animation/ContextualAnimation/Source/ContextualAnimation/Public/ContextualAnimTypes.h
@@ -365,6 +365,10 @@ private:
 	FContextualAnimSceneBindingContext Context;
 
 	int32 AnimTrackIdx = INDEX_NONE;
+
+	mutable TObjectPtr<UContextualAnimSceneActorComponent> CachedSceneActorComp = nullptr;
+	mutable TObjectPtr<UAnimInstance> CachedAnimInstance = nullptr;
+	mutable TObjectPtr<USkeletalMeshComponent> CachedSkeletalMesh = nullptr;
 };
 
 template<>
