Contextual Anim: Added option to add actors to bindings one by one. SceneComp now forces net update when an interaction starts and ends and other minor fixes.

diff --git a/Engine/Plugins/Experimental/Animation/ContextualAnimation/Source/ContextualAnimation/Private/ContextualAnimSceneActorComponent.cpp b/Engine/Plugins/Experimental/Animation/ContextualAnimation/Source/ContextualAnimation/Private/ContextualAnimSceneActorComponent.cpp
--- a/Engine/Plugins/Experimental/Animation/ContextualAnimation/Source/ContextualAnimation/Private/ContextualAnimSceneActorComponent.cpp
+++ b/Engine/Plugins/Experimental/Animation/ContextualAnimation/Source/ContextualAnimation/Private/ContextualAnimSceneActorComponent.cpp
@@ -68,6 +68,8 @@ bool UContextualAnimSceneActorComponent::StartContextualAnimScene(const FContext
 		{
 			RepBindings = InBindings;
 			MARK_PROPERTY_DIRTY_FROM_NAME(UContextualAnimSceneActorComponent, RepBindings, this);
+
+			GetOwner()->ForceNetUpdate();
 		}
 
 		return true;
@@ -78,8 +80,8 @@ bool UContextualAnimSceneActorComponent::StartContextualAnimScene(const FContext
 
 void UContextualAnimSceneActorComponent::OnRep_Bindings(const FContextualAnimSceneBindings& LastRepBindings)
 {
-	// @TODO: This need more investigation but for now it prevents an issue caused by this OnRep_ triggering even when there is no (obvious) changein the data
-	if(RepBindings.GetID() == LastRepBindings.GetID())
+	// @TODO: This need more investigation but for now it prevents an issue caused by this OnRep_ triggering even when there is no (obvious) change in the data
+	if(RepBindings.GetID() == LastRepBindings.GetID() && RepBindings.IsValid() && LastRepBindings.IsValid())
 	{
 		UE_LOG(LogContextualAnim, Warning, TEXT("%-21s UContextualAnimSceneActorComponent::OnRep_Bindings Actor: %s RepBindings Id: %d LastRepBindings Id: %d"),
 			*UEnum::GetValueAsString(TEXT("Engine.ENetRole"), GetOwner()->GetLocalRole()), *GetNameSafe(GetOwner()), RepBindings.GetID(), LastRepBindings.GetID());
@@ -87,8 +89,8 @@ void UContextualAnimSceneActorComponent::OnRep_Bindings(const FContextualAnimSce
 		return;
 	}
 
-	UE_LOG(LogContextualAnim, Verbose, TEXT("%-21s UContextualAnimSceneActorComponent::OnRep_Bindings Actor: %s RepBindings Id: %d Bindings Id: %d"),
-		*UEnum::GetValueAsString(TEXT("Engine.ENetRole"), GetOwner()->GetLocalRole()), *GetNameSafe(GetOwner()), RepBindings.GetID(), Bindings.GetID());
+	UE_LOG(LogContextualAnim, Verbose, TEXT("%-21s UContextualAnimSceneActorComponent::OnRep_Bindings Actor: %s RepBindings Id: %d Num: %d Bindings Id: %d Num: %d"),
+		*UEnum::GetValueAsString(TEXT("Engine.ENetRole"), GetOwner()->GetLocalRole()), *GetNameSafe(GetOwner()), RepBindings.GetID(), RepBindings.Num(), Bindings.GetID(), Bindings.Num());
 
 	// The owner of this component started an interaction on the server
 	if (RepBindings.IsValid())
@@ -118,7 +120,10 @@ void UContextualAnimSceneActorComponent::OnRep_Bindings(const FContextualAnimSce
 		// In this case we don't want to tell everyone else to also leave the scene since there is very common for the initiator, 
 		// specially if is player character, to end the animation earlier for responsiveness
 		// It is more likely this will do nothing since we listen to montage end also on Simulated Proxies to 'predict' the end of the interaction.
-		LeaveScene();
+		if(RepBindings.GetID() == Bindings.GetID())
+		{
+			LeaveScene();
+		}
 	}
 }
 
@@ -196,8 +201,12 @@ void UContextualAnimSceneActorComponent::OnJoinedScene(const FContextualAnimScen
 		// Prevent physics rotation. During the interaction we want to be fully root motion driven
 		if (UCharacterMovementComponent* MovementComp = GetOwner()->FindComponentByClass<UCharacterMovementComponent>())
 		{
-			bAllowPhysicsRotationDuringAnimRootMotionBackup = MovementComp->bAllowPhysicsRotationDuringAnimRootMotion;
+			RotationPropertiesBackup.bAllowPhysicsRotationDuringAnimRootMotion = MovementComp->bAllowPhysicsRotationDuringAnimRootMotion;
+			RotationPropertiesBackup.bUseControllerDesiredRotation = MovementComp->bUseControllerDesiredRotation;
+			RotationPropertiesBackup.bOrientRotationToMovement = MovementComp->bOrientRotationToMovement;
 			MovementComp->bAllowPhysicsRotationDuringAnimRootMotion = false;
+			MovementComp->bUseControllerDesiredRotation = false;
+			MovementComp->bOrientRotationToMovement = false;
 		}
 
 		OnJoinedSceneDelegate.Broadcast(this);
@@ -226,7 +235,9 @@ void UContextualAnimSceneActorComponent::OnLeftScene()
 		// Restore bAllowPhysicsRotationDuringAnimRootMotion
 		if (UCharacterMovementComponent* MovementComp = GetOwner()->FindComponentByClass<UCharacterMovementComponent>())
 		{
-			MovementComp->bAllowPhysicsRotationDuringAnimRootMotion = bAllowPhysicsRotationDuringAnimRootMotionBackup;
+			MovementComp->bAllowPhysicsRotationDuringAnimRootMotion = RotationPropertiesBackup.bAllowPhysicsRotationDuringAnimRootMotion;
+			MovementComp->bUseControllerDesiredRotation = RotationPropertiesBackup.bUseControllerDesiredRotation;
+			MovementComp->bOrientRotationToMovement = RotationPropertiesBackup.bOrientRotationToMovement;
 		}
 
 		OnLeftSceneDelegate.Broadcast(this);
@@ -244,8 +255,8 @@ void UContextualAnimSceneActorComponent::JoinScene(const FContextualAnimSceneBin
 
 	if (const FContextualAnimSceneBinding* Binding = InBindings.FindBindingByActor(GetOwner()))
 	{
-		UE_LOG(LogContextualAnim, Verbose, TEXT("%-21s UContextualAnimSceneActorComponent::JoinScene Actor: %s InBindings Id: %d Section: %d Asset: %s"),
-			*UEnum::GetValueAsString(TEXT("Engine.ENetRole"), GetOwner()->GetLocalRole()), *GetNameSafe(GetOwner()), InBindings.GetID(), InBindings.GetSectionIdx(), *GetNameSafe(InBindings.GetSceneAsset()));
+		UE_LOG(LogContextualAnim, Verbose, TEXT("%-21s UContextualAnimSceneActorComponent::JoinScene Actor: %s Role: %s InBindings Id: %d Section: %d Asset: %s"),
+			*UEnum::GetValueAsString(TEXT("Engine.ENetRole"), GetOwner()->GetLocalRole()), *GetNameSafe(GetOwner()), *InBindings.GetRoleFromBinding(*Binding).ToString(), InBindings.GetID(), InBindings.GetSectionIdx(), *GetNameSafe(InBindings.GetSceneAsset()));
 
 		Bindings = InBindings;
 
@@ -272,8 +283,12 @@ void UContextualAnimSceneActorComponent::JoinScene(const FContextualAnimSceneBin
 		// Prevent physics rotation. During the interaction we want to be fully root motion driven
 		if (UCharacterMovementComponent* MovementComp = GetOwner()->FindComponentByClass<UCharacterMovementComponent>())
 		{
-			bAllowPhysicsRotationDuringAnimRootMotionBackup = MovementComp->bAllowPhysicsRotationDuringAnimRootMotion;
+			RotationPropertiesBackup.bAllowPhysicsRotationDuringAnimRootMotion = MovementComp->bAllowPhysicsRotationDuringAnimRootMotion;
+			RotationPropertiesBackup.bUseControllerDesiredRotation = MovementComp->bUseControllerDesiredRotation;
+			RotationPropertiesBackup.bOrientRotationToMovement = MovementComp->bOrientRotationToMovement;
 			MovementComp->bAllowPhysicsRotationDuringAnimRootMotion = false;
+			MovementComp->bUseControllerDesiredRotation = false;
+			MovementComp->bOrientRotationToMovement = false;
 
 			//@TODO: Temp solution that assumes these interactions are not locally predicted and that is ok to be in flying mode during the entire animation
 			if (AnimTrack.bRequireFlyingMode && MovementComp->MovementMode != MOVE_Flying)
@@ -290,8 +305,9 @@ void UContextualAnimSceneActorComponent::LeaveScene()
 {
 	if (const FContextualAnimSceneBinding* Binding = Bindings.FindBindingByActor(GetOwner()))
 	{
-		UE_LOG(LogContextualAnim, Verbose, TEXT("%-21s UContextualAnimSceneActorComponent::LeaveScene Actor: %s Current Bindings Id: %d Section: %d Asset: %s"),
-			*UEnum::GetValueAsString(TEXT("Engine.ENetRole"), GetOwner()->GetLocalRole()), *GetNameSafe(GetOwner()), Bindings.GetID(), Bindings.GetSectionIdx(), *GetNameSafe(Bindings.GetSceneAsset()));
+		UE_LOG(LogContextualAnim, Verbose, TEXT("%-21s UContextualAnimSceneActorComponent::LeaveScene Actor: %s Role: %s Current Bindings Id: %d Section: %d Asset: %s"),
+			*UEnum::GetValueAsString(TEXT("Engine.ENetRole"), GetOwner()->GetLocalRole()), *GetNameSafe(GetOwner()), *Bindings.GetRoleFromBinding(*Binding).ToString(),
+			Bindings.GetID(), Bindings.GetSectionIdx(), *GetNameSafe(Bindings.GetSceneAsset()));
 
 		const FContextualAnimTrack& AnimTrack = Bindings.GetAnimTrackFromBinding(*Binding);
 
@@ -314,7 +330,7 @@ void UContextualAnimSceneActorComponent::LeaveScene()
 		{
 			SkelMeshComp->OnTickPose.RemoveAll(this);
 		}
-
+		
 		// Restore collision between actors
 		// Note that this assumes that we are the only one disabling the collision between these actors. 
 		// We might want to add a more robust mechanism to avoid overriding a request to disable collision that may have been set by another system
@@ -323,7 +339,9 @@ void UContextualAnimSceneActorComponent::LeaveScene()
 		// Restore bAllowPhysicsRotationDuringAnimRootMotion
 		if (UCharacterMovementComponent* MovementComp = GetOwner()->FindComponentByClass<UCharacterMovementComponent>())
 		{
-			MovementComp->bAllowPhysicsRotationDuringAnimRootMotion = bAllowPhysicsRotationDuringAnimRootMotionBackup;
+			MovementComp->bAllowPhysicsRotationDuringAnimRootMotion = RotationPropertiesBackup.bAllowPhysicsRotationDuringAnimRootMotion;
+			MovementComp->bUseControllerDesiredRotation = RotationPropertiesBackup.bUseControllerDesiredRotation;
+			MovementComp->bOrientRotationToMovement = RotationPropertiesBackup.bOrientRotationToMovement;
 
 			//@TODO: Temp solution that assumes these interactions are not locally predicted and that is ok to be in flying mode during the entire animation
 			if (AnimTrack.bRequireFlyingMode && MovementComp->MovementMode == MOVE_Flying)
@@ -354,8 +372,10 @@ void UContextualAnimSceneActorComponent::OnMontageBlendingOut(UAnimMontage* Mont
 				// Rep empty bindings if we were the initiator of this interaction.
 				if(RepBindings.IsValid())
 				{
-					RepBindings.Reset();
+					RepBindings.Clear();
 					MARK_PROPERTY_DIRTY_FROM_NAME(UContextualAnimSceneActorComponent, RepBindings, this);
+
+					GetOwner()->ForceNetUpdate();
 				}
 
 				//@TODO: Replicate this event separately for each other member of the interaction
diff --git a/Engine/Plugins/Experimental/Animation/ContextualAnimation/Source/ContextualAnimation/Private/ContextualAnimTypes.cpp b/Engine/Plugins/Experimental/Animation/ContextualAnimation/Source/ContextualAnimation/Private/ContextualAnimTypes.cpp
--- a/Engine/Plugins/Experimental/Animation/ContextualAnimation/Source/ContextualAnimation/Private/ContextualAnimTypes.cpp
+++ b/Engine/Plugins/Experimental/Animation/ContextualAnimation/Source/ContextualAnimation/Private/ContextualAnimTypes.cpp
@@ -369,7 +369,11 @@ FContextualAnimSceneBindings::FContextualAnimSceneBindings(const UContextualAnim
 	SceneAsset = &InSceneAsset;
 	SectionIdx = InSectionIdx;
 	AnimSetIdx = InAnimSetIdx;
+	GenerateUniqueId();
+}
 
+void FContextualAnimSceneBindings::GenerateUniqueId()
+{
 	static uint8 IncrementID = 0;
 	IncrementID = IncrementID < UINT8_MAX ? IncrementID + 1 : 0;
 	Id = IncrementID;
@@ -390,12 +394,40 @@ void FContextualAnimSceneBindings::Reset()
 	SceneInstancePtr.Reset();
 }
 
+void FContextualAnimSceneBindings::Clear()
+{
+	Data.Reset();
+	SceneInstancePtr.Reset();
+}
+
 const FContextualAnimSceneBinding* FContextualAnimSceneBindings::GetSyncLeader() const
 {
 	//@TODO: Return first secondary binding as sync leader for now. This may have to be explicitly defined, either in the SceneAsset or when creating the bindings.
 	return Data.FindByPredicate([this](const FContextualAnimSceneBinding& Item) { return GetRoleFromBinding(Item) != SceneAsset->GetPrimaryRole(); });
 }
 
+bool FContextualAnimSceneBindings::BindActorToRole(AActor& ActorRef, FName Role)
+{
+	if(const FContextualAnimSceneBinding* Binding = FindBindingByRole(Role))
+	{
+		UE_LOG(LogContextualAnim, Warning, TEXT("FContextualAnimSceneBindings::BindActorToRole. Failed to bind Actor: '%s' to Role: '%s'. Reason: %s already bound"), 
+			*GetNameSafe(&ActorRef), *Role.ToString(), *GetNameSafe(Binding->GetActor()));
+		return false;
+	}
+
+	const FContextualAnimTrack* AnimTrackPtr = SceneAsset->GetAnimTrack(SectionIdx, AnimSetIdx, Role);
+	if(AnimTrackPtr == nullptr)
+	{
+		UE_LOG(LogContextualAnim, Warning, TEXT("FContextualAnimSceneBindings::BindActorToRole. Failed to bind Actor: '%s' to Role: '%s'. Reason: Can't find valid AnimTrack for it. SceneAsset: %s SectionIdx: %d AnimSetIdx: %d"),
+			*GetNameSafe(&ActorRef), *Role.ToString(), *GetNameSafe(SceneAsset.Get()), SectionIdx, AnimSetIdx);
+		return false;
+	}
+
+	Data.Add(FContextualAnimSceneBinding(FContextualAnimSceneBindingContext(&ActorRef), *AnimTrackPtr));
+
+	return true;
+}
+
 const FContextualAnimTrack& FContextualAnimSceneBindings::GetAnimTrackFromBinding(const FContextualAnimSceneBinding& Binding) const
 {
 	if (const FContextualAnimTrack* AnimTrack = GetSceneAsset()->GetAnimTrack(SectionIdx, AnimSetIdx, Binding.AnimTrackIdx))
diff --git a/Engine/Plugins/Experimental/Animation/ContextualAnimation/Source/ContextualAnimation/Public/ContextualAnimSceneActorComponent.h b/Engine/Plugins/Experimental/Animation/ContextualAnimation/Source/ContextualAnimation/Public/ContextualAnimSceneActorComponent.h
--- a/Engine/Plugins/Experimental/Animation/ContextualAnimation/Source/ContextualAnimation/Public/ContextualAnimSceneActorComponent.h
+++ b/Engine/Plugins/Experimental/Animation/ContextualAnimation/Source/ContextualAnimation/Public/ContextualAnimSceneActorComponent.h
@@ -83,9 +83,13 @@ protected:
 	UPROPERTY(Transient)
 	TArray<FContextualAnimIKTarget> IKTargets;
 
-	/** Value of AllowPhysicsRotationDuringAnimRootMotion property for this actor before we join an scene, so we can restore it once the interaction ends */
-	UPROPERTY(Transient)
-	bool bAllowPhysicsRotationDuringAnimRootMotionBackup = false;
+	struct FRotationProperties
+	{
+		bool bAllowPhysicsRotationDuringAnimRootMotion = false;
+		bool bUseControllerDesiredRotation = false;
+		bool bOrientRotationToMovement = false;
+	};
+	FRotationProperties RotationPropertiesBackup;
 
 	void UpdateIKTargets();
 
diff --git a/Engine/Plugins/Experimental/Animation/ContextualAnimation/Source/ContextualAnimation/Public/ContextualAnimTypes.h b/Engine/Plugins/Experimental/Animation/ContextualAnimation/Source/ContextualAnimation/Public/ContextualAnimTypes.h
--- a/Engine/Plugins/Experimental/Animation/ContextualAnimation/Source/ContextualAnimation/Public/ContextualAnimTypes.h
+++ b/Engine/Plugins/Experimental/Animation/ContextualAnimation/Source/ContextualAnimation/Public/ContextualAnimTypes.h
@@ -417,6 +417,8 @@ struct CONTEXTUALANIMATION_API FContextualAnimSceneBindings
 	static bool TryCreateBindings(const UContextualAnimSceneAsset& SceneAsset, int32 SectionIdx, int32 AnimSetIdx, const FContextualAnimSceneBindingContext& Primary, const FContextualAnimSceneBindingContext& Secondary, FContextualAnimSceneBindings& OutBindings);
 	static bool TryCreateBindings(const UContextualAnimSceneAsset& SceneAsset, int32 SectionIdx, const FContextualAnimSceneBindingContext& Primary, const FContextualAnimSceneBindingContext& Secondary, FContextualAnimSceneBindings& OutBindings);
 
+	bool BindActorToRole(AActor& ActorRef, FName Role);
+
 	void CalculateAnimSetPivots(TArray<FContextualAnimSetPivot>& OutScenePivots) const;
 	bool CalculateAnimSetPivot(const FContextualAnimSetPivotDefinition& AnimSetPivotDef, FContextualAnimSetPivot& OutScenePivot) const;
 
@@ -432,8 +434,12 @@ struct CONTEXTUALANIMATION_API FContextualAnimSceneBindings
 
 	void Reset();
 
+	void Clear();
+
 	bool NetSerialize(FArchive& Ar, class UPackageMap* Map, bool& bOutSuccess);
 
+	void GenerateUniqueId();
+
 private:
 
 	friend class UContextualAnimManager;
