ContextualAnim: BindActorToRole will now fail if the actor doesn't have a scene actor comp

diff --git a/Engine/Plugins/Experimental/Animation/ContextualAnimation/Source/ContextualAnimation/Private/ContextualAnimSceneActorComponent.cpp b/Engine/Plugins/Experimental/Animation/ContextualAnimation/Source/ContextualAnimation/Private/ContextualAnimSceneActorComponent.cpp
--- a/Engine/Plugins/Experimental/Animation/ContextualAnimation/Source/ContextualAnimation/Private/ContextualAnimSceneActorComponent.cpp
+++ b/Engine/Plugins/Experimental/Animation/ContextualAnimation/Source/ContextualAnimation/Private/ContextualAnimSceneActorComponent.cpp
@@ -51,13 +51,16 @@ bool UContextualAnimSceneActorComponent::StartContextualAnimScene(const FContext
 	const FContextualAnimSceneBinding* OwnerBinding = InBindings.FindBindingByActor(GetOwner());
 	if (ensureAlways(OwnerBinding))
 	{
-		OwnerBinding->GetSceneActorComponent()->JoinScene(InBindings);
+		JoinScene(InBindings);
 
 		for (const FContextualAnimSceneBinding& Binding : InBindings)
 		{
 			if (Binding.GetActor() != GetOwner())
 			{
-				Binding.GetSceneActorComponent()->JoinScene(InBindings);
+				UContextualAnimSceneActorComponent* Comp = Binding.GetSceneActorComponent();
+				checkf(Comp, TEXT("Missing SceneActorComp on %s"), *GetNameSafe(Binding.GetActor()));
+
+				Comp->JoinScene(InBindings);
 			}
 		}
 
@@ -99,7 +102,7 @@ void UContextualAnimSceneActorComponent::OnRep_Bindings(const FContextualAnimSce
 		if (ensureAlways(OwnerBinding))
 		{
 			// Join the scene (start playing animation, etc.)
-			OwnerBinding->GetSceneActorComponent()->JoinScene(RepBindings);
+			JoinScene(RepBindings);
 
 			// RepBindings is only replicated from the initiator of the actor. 
 			// So now we have to tell everyone else involved in the interaction to join us
@@ -109,7 +112,10 @@ void UContextualAnimSceneActorComponent::OnRep_Bindings(const FContextualAnimSce
 			{
 				if (Binding.GetActor() != GetOwner())
 				{
-					Binding.GetSceneActorComponent()->JoinScene(RepBindings);
+					UContextualAnimSceneActorComponent* Comp = Binding.GetSceneActorComponent();
+					checkf(Comp, TEXT("Missing SceneActorComp on %s"), *GetNameSafe(Binding.GetActor()));
+
+					Comp->JoinScene(RepBindings);
 				}
 			}
 		}
diff --git a/Engine/Plugins/Experimental/Animation/ContextualAnimation/Source/ContextualAnimation/Private/ContextualAnimTypes.cpp b/Engine/Plugins/Experimental/Animation/ContextualAnimation/Source/ContextualAnimation/Private/ContextualAnimTypes.cpp
--- a/Engine/Plugins/Experimental/Animation/ContextualAnimation/Source/ContextualAnimation/Private/ContextualAnimTypes.cpp
+++ b/Engine/Plugins/Experimental/Animation/ContextualAnimation/Source/ContextualAnimation/Private/ContextualAnimTypes.cpp
@@ -408,6 +408,14 @@ const FContextualAnimSceneBinding* FContextualAnimSceneBindings::GetSyncLeader()
 
 bool FContextualAnimSceneBindings::BindActorToRole(AActor& ActorRef, FName Role)
 {
+	const UContextualAnimSceneActorComponent* Comp = ActorRef.FindComponentByClass<UContextualAnimSceneActorComponent>();
+	if (Comp == nullptr)
+	{
+		UE_LOG(LogContextualAnim, Warning, TEXT("FContextualAnimSceneBindings::BindActorToRole. Failed to bind Actor: '%s' to Role: '%s'. Reason: Missing SceneActorComp"),
+			*GetNameSafe(&ActorRef), *Role.ToString());
+		return false;
+	}
+
 	if(const FContextualAnimSceneBinding* Binding = FindBindingByRole(Role))
 	{
 		UE_LOG(LogContextualAnim, Warning, TEXT("FContextualAnimSceneBindings::BindActorToRole. Failed to bind Actor: '%s' to Role: '%s'. Reason: %s already bound"), 
