Contextual Anim: Added option to override preview actor in the preview scene

diff --git a/Engine/Plugins/Experimental/Animation/ContextualAnimation/Source/ContextualAnimation/Public/ContextualAnimSceneAsset.h b/Engine/Plugins/Experimental/Animation/ContextualAnimation/Source/ContextualAnimation/Public/ContextualAnimSceneAsset.h
--- a/Engine/Plugins/Experimental/Animation/ContextualAnimation/Source/ContextualAnimation/Public/ContextualAnimSceneAsset.h
+++ b/Engine/Plugins/Experimental/Animation/ContextualAnimation/Source/ContextualAnimation/Public/ContextualAnimSceneAsset.h
@@ -140,6 +140,36 @@ enum class EContextualAnimCriterionToConsider : uint8
 	Other
 };
 
+UENUM(BlueprintType)
+enum class EContextualAnimActorPreviewType : uint8
+{
+	SkeletalMesh,
+	StaticMesh,
+	Actor,
+	None
+};
+
+USTRUCT(BlueprintType)
+struct FContextualAnimActorPreviewData
+{
+	GENERATED_BODY()
+
+	UPROPERTY(EditAnywhere, Category = "Defaults", meta = (GetOptions = "GetRoles"))
+	FName Role;
+
+	UPROPERTY(EditAnywhere, Category = "Defaults", meta = (GetOptions = "GetRoles"))
+	EContextualAnimActorPreviewType Type = EContextualAnimActorPreviewType::StaticMesh;
+
+	UPROPERTY(EditAnywhere, Category = "Defaults", meta = (EditCondition = "Type==EContextualAnimActorPreviewType::SkeletalMesh", EditConditionHides))
+	TSoftObjectPtr<class USkeletalMesh> PreviewSkeletalMesh;
+
+	UPROPERTY(EditAnywhere, Category = "Defaults", meta = (EditCondition = "Type==EContextualAnimActorPreviewType::StaticMesh", EditConditionHides))
+	TSoftObjectPtr<class UStaticMesh> PreviewStaticMesh;
+
+	UPROPERTY(EditAnywhere, Category = "Defaults", meta = (EditCondition = "Type==EContextualAnimActorPreviewType::Actor", EditConditionHides))
+	TSoftClassPtr<class AActor> PreviewActorClass;
+};
+
 UCLASS(Blueprintable)
 class CONTEXTUALANIMATION_API UContextualAnimSceneAsset : public UDataAsset
 {
@@ -241,6 +271,9 @@ protected:
 	UPROPERTY(EditAnywhere, Category = "Defaults", meta = (GetOptions = "GetRoles"))
 	FName PrimaryRole = NAME_None;
 
+	UPROPERTY(EditAnywhere, Category = "Defaults")
+	TArray<FContextualAnimActorPreviewData> OverridePreviewData;
+
 	UPROPERTY(EditAnywhere, Category = "Defaults")
 	TArray<FContextualAnimSceneSection> Sections;
 
diff --git a/Engine/Plugins/Experimental/Animation/ContextualAnimation/Source/ContextualAnimationEditor/Private/ContextualAnimViewModel.cpp b/Engine/Plugins/Experimental/Animation/ContextualAnimation/Source/ContextualAnimationEditor/Private/ContextualAnimViewModel.cpp
--- a/Engine/Plugins/Experimental/Animation/ContextualAnimation/Source/ContextualAnimationEditor/Private/ContextualAnimViewModel.cpp
+++ b/Engine/Plugins/Experimental/Animation/ContextualAnimation/Source/ContextualAnimationEditor/Private/ContextualAnimViewModel.cpp
@@ -141,86 +141,143 @@ void FContextualAnimViewModel::SetActiveAnimSetForSection(int32 SectionIdx, int3
 
 AActor* FContextualAnimViewModel::SpawnPreviewActor(const FContextualAnimTrack& AnimTrack)
 {
-	if (AnimTrack.Animation == nullptr)
-	{
-		return nullptr;
-	}
+	AActor* PreviewActor = nullptr;
 
-	USkeleton* Skeleton = AnimTrack.Animation->GetSkeleton();
+	// Prioritize override preview
+	EContextualAnimActorPreviewType PreviewType = EContextualAnimActorPreviewType::None;
+	USkeletalMesh* PreviewSkeletalMesh = nullptr;
+	UStaticMesh* PreviewStaticMesh = nullptr;
+	UClass* PreviewActorClass = nullptr;
 
-	if (Skeleton == nullptr)
+	for (const FContextualAnimActorPreviewData& PreviewData : GetSceneAsset()->OverridePreviewData)
 	{
-		return nullptr;
+		if (PreviewData.Role == AnimTrack.Role)
+		{
+			if(PreviewData.Type == EContextualAnimActorPreviewType::StaticMesh)
+			{
+				if(UStaticMesh* StaticMesh = PreviewData.PreviewStaticMesh.LoadSynchronous())
+				{
+					PreviewStaticMesh = StaticMesh;
+					PreviewType = PreviewData.Type;
+				}
+			}
+			else if (PreviewData.Type == EContextualAnimActorPreviewType::SkeletalMesh)
+			{
+				if (USkeletalMesh* SkeletalMesh = PreviewData.PreviewSkeletalMesh.LoadSynchronous())
+				{
+					PreviewSkeletalMesh = SkeletalMesh;
+					PreviewType = PreviewData.Type;
+				}
+			}
+			else if (PreviewData.Type == EContextualAnimActorPreviewType::Actor)
+			{
+				if (UClass* ActorClass = PreviewData.PreviewActorClass.LoadSynchronous())
+				{
+					PreviewActorClass = ActorClass;
+					PreviewType = PreviewData.Type;
+				}
+			}
+			break;
+		}
 	}
 
-	const FContextualAnimRoleDefinition* RoleDef = GetSceneAsset()->RolesAsset ? GetSceneAsset()->RolesAsset->FindRoleDefinitionByName(AnimTrack.Role) : nullptr;
-	const bool bIsCharacter = (RoleDef && RoleDef->bIsCharacter);
+	// if not explicit preview mesh is defined for a role with a valid animation, try to pull the preview mesh from the animation
+	if (PreviewType == EContextualAnimActorPreviewType::None && AnimTrack.Animation && AnimTrack.Animation->GetSkeleton())
+	{
+		PreviewSkeletalMesh = AnimTrack.Animation->GetSkeleton()->GetPreviewMesh(true);
+
+		if (PreviewSkeletalMesh)
+		{
+			PreviewType = EContextualAnimActorPreviewType::SkeletalMesh;
+		}
+	}
 
 	const FTransform SpawnTransform = AnimTrack.GetRootTransformAtTime(0.f);
 
 	FActorSpawnParameters Params;
 	Params.SpawnCollisionHandlingOverride = ESpawnActorCollisionHandlingMethod::AdjustIfPossibleButAlwaysSpawn;
 
-	if (bIsCharacter)
+	if (PreviewType != EContextualAnimActorPreviewType::None)
 	{
-		ACharacter* PreviewCharacter = GetWorld()->SpawnActor<ACharacter>(ACharacter::StaticClass(), SpawnTransform, Params);
-		PreviewCharacter->SetFlags(RF_Transient);
+		const FContextualAnimRoleDefinition* RoleDef = GetSceneAsset()->RolesAsset ? GetSceneAsset()->RolesAsset->FindRoleDefinitionByName(AnimTrack.Role) : nullptr;
+		const bool bIsCharacter = (RoleDef && RoleDef->bIsCharacter);
 
-		USkeletalMeshComponent* SkelMeshComp = PreviewCharacter->GetMesh();
-		SkelMeshComp->SetRelativeLocation(FVector(0.f, 0.f, -PreviewCharacter->GetCapsuleComponent()->GetScaledCapsuleHalfHeight()));
-		SkelMeshComp->SetRelativeRotation(RoleDef->MeshToComponent.GetRotation());
-		SkelMeshComp->SetSkeletalMesh(Skeleton->GetPreviewMesh(true));
-		SkelMeshComp->SetAnimationMode(EAnimationMode::AnimationBlueprint);
-		SkelMeshComp->SetAnimInstanceClass(UAnimPreviewInstance::StaticClass());
+		if (bIsCharacter)
+		{
+			ACharacter* PreviewCharacter = GetWorld()->SpawnActor<ACharacter>(ACharacter::StaticClass(), SpawnTransform, Params);
+			PreviewCharacter->SetFlags(RF_Transient);
 
-		UAnimPreviewInstance* AnimInstance = Cast<UAnimPreviewInstance>(SkelMeshComp->GetAnimInstance());
-		AnimInstance->SetAnimationAsset(AnimTrack.Animation, false, 0.f);
-		AnimInstance->PlayAnim(false, 0.f);
+			USkeletalMeshComponent* SkelMeshComp = PreviewCharacter->GetMesh();
+			SkelMeshComp->SetRelativeLocation(FVector(0.f, 0.f, -PreviewCharacter->GetCapsuleComponent()->GetScaledCapsuleHalfHeight()));
+			SkelMeshComp->SetRelativeRotation(RoleDef->MeshToComponent.GetRotation());
+			SkelMeshComp->SetSkeletalMesh(PreviewSkeletalMesh);
+			SkelMeshComp->SetAnimationMode(EAnimationMode::AnimationBlueprint);
+			SkelMeshComp->SetAnimInstanceClass(UAnimPreviewInstance::StaticClass());
 
-		PreviewCharacter->CacheInitialMeshOffset(SkelMeshComp->GetRelativeLocation(), SkelMeshComp->GetRelativeRotation());
+			UAnimPreviewInstance* AnimInstance = Cast<UAnimPreviewInstance>(SkelMeshComp->GetAnimInstance());
+			AnimInstance->SetAnimationAsset(AnimTrack.Animation, false, 0.f);
+			AnimInstance->PlayAnim(false, 0.f);
 
-		if (UCharacterMovementComponent* CharacterMovementComp = PreviewCharacter->GetCharacterMovement())
-		{
-			CharacterMovementComp->bOrientRotationToMovement = true;
-			CharacterMovementComp->bUseControllerDesiredRotation = false;
-			CharacterMovementComp->RotationRate = FRotator(0.f, 540.0, 0.f);
-			CharacterMovementComp->bRunPhysicsWithNoController = true;
+			PreviewCharacter->CacheInitialMeshOffset(SkelMeshComp->GetRelativeLocation(), SkelMeshComp->GetRelativeRotation());
 
-			CharacterMovementComp->SetMovementMode(AnimTrack.bRequireFlyingMode ? EMovementMode::MOVE_Flying : EMovementMode::MOVE_Walking);
-		}
+			if (UCharacterMovementComponent* CharacterMovementComp = PreviewCharacter->GetCharacterMovement())
+			{
+				CharacterMovementComp->bOrientRotationToMovement = true;
+				CharacterMovementComp->bUseControllerDesiredRotation = false;
+				CharacterMovementComp->RotationRate = FRotator(0.f, 540.0, 0.f);
+				CharacterMovementComp->bRunPhysicsWithNoController = true;
 
-		UMotionWarpingComponent* MotionWarpingComp = NewObject<UMotionWarpingComponent>(PreviewCharacter);
-		MotionWarpingComp->RegisterComponentWithWorld(GetWorld());
-		MotionWarpingComp->InitializeComponent();
+				CharacterMovementComp->SetMovementMode(AnimTrack.bRequireFlyingMode ? EMovementMode::MOVE_Flying : EMovementMode::MOVE_Walking);
+			}
 
-		return PreviewCharacter;
-	}
-	else
-	{
-		AActor* PreviewActor = GetWorld()->SpawnActor<AActor>(AActor::StaticClass(), SpawnTransform, Params);
-		PreviewActor->SetFlags(RF_Transient);
+			UMotionWarpingComponent* MotionWarpingComp = NewObject<UMotionWarpingComponent>(PreviewCharacter);
+			MotionWarpingComp->RegisterComponentWithWorld(GetWorld());
+			MotionWarpingComp->InitializeComponent();
+
+			PreviewActor = PreviewCharacter;
+		}
+		else
+		{
+			if(PreviewActorClass)
+			{
+				PreviewActor = GetWorld()->SpawnActor<AActor>(PreviewActorClass, SpawnTransform, Params);
+				PreviewActor->SetFlags(RF_Transient);
+			}
+			else 
+			{
+				PreviewActor = GetWorld()->SpawnActor<AActor>(AActor::StaticClass(), SpawnTransform, Params);
+				PreviewActor->SetFlags(RF_Transient);
 
-		UDebugSkelMeshComponent* SkelMeshComp = NewObject<UDebugSkelMeshComponent>(PreviewActor);
-		SkelMeshComp->RegisterComponentWithWorld(GetWorld());
+				if (PreviewSkeletalMesh)
+				{
+					UDebugSkelMeshComponent* SkelMeshComp = NewObject<UDebugSkelMeshComponent>(PreviewActor);
+					SkelMeshComp->RegisterComponentWithWorld(GetWorld());
 
-		UAnimPreviewInstance* AnimInstance = NewObject<UAnimPreviewInstance>(SkelMeshComp);
-		SkelMeshComp->PreviewInstance = AnimInstance;
-		AnimInstance->InitializeAnimation();
+					UAnimPreviewInstance* AnimInstance = NewObject<UAnimPreviewInstance>(SkelMeshComp);
+					SkelMeshComp->PreviewInstance = AnimInstance;
+					AnimInstance->InitializeAnimation();
 
-		SkelMeshComp->SetSkeletalMesh(Skeleton->GetPreviewMesh(true));
-		SkelMeshComp->EnablePreview(true, AnimTrack.Animation);
+					SkelMeshComp->SetSkeletalMesh(PreviewSkeletalMesh);
+					SkelMeshComp->EnablePreview(true, AnimTrack.Animation);
 
-		AnimInstance->SetAnimationAsset(AnimTrack.Animation, false, 0.0f);
+					AnimInstance->SetAnimationAsset(AnimTrack.Animation, false, 0.0f);
 
-		AnimInstance->PlayAnim(false, 0.0f);
+					AnimInstance->PlayAnim(false, 0.0f);
 
-		if (!PreviewActor->GetRootComponent())
-		{
-			PreviewActor->SetRootComponent(SkelMeshComp);
+					PreviewActor->SetRootComponent(SkelMeshComp);
+				}
+				else if (PreviewStaticMesh)
+				{
+					UStaticMeshComponent* StaticMeshComp = NewObject<UStaticMeshComponent>(PreviewActor);
+					StaticMeshComp->SetStaticMesh(PreviewStaticMesh);
+					StaticMeshComp->RegisterComponentWithWorld(GetWorld());
+					PreviewActor->SetRootComponent(StaticMeshComp);
+				}
+			}
 		}
-
-		return PreviewActor;
 	}
+
+	return PreviewActor;
 }
 
 void FContextualAnimViewModel::ResetTimeline()
@@ -1020,13 +1077,18 @@ void FContextualAnimViewModel::OnFinishedChangingProperties(const FPropertyChang
 	const FName PropertyName = (PropertyChangedEvent.Property != nullptr) ? PropertyChangedEvent.Property->GetFName() : NAME_None;
 	const FName MemberPropertyName = (PropertyChangedEvent.MemberProperty != nullptr) ? PropertyChangedEvent.MemberProperty->GetFName() : NAME_None;
 
-	// Refresh the time line if the Sections array or the AnimSets array inside a section changes
+	// Refresh preview scene if the Sections array or the AnimSets array inside a section changes
 	if (MemberPropertyName == GET_MEMBER_NAME_CHECKED(UContextualAnimSceneAsset, Sections) &&
 		(PropertyName == GET_MEMBER_NAME_CHECKED(UContextualAnimSceneAsset, Sections) ||
 			PropertyName == GET_MEMBER_NAME_CHECKED(FContextualAnimSceneSection, AnimSets)))
 	{
 		SetDefaultMode();
 	}
+	// Refresh preview scene if override preview data changes
+	else if (MemberPropertyName == GET_MEMBER_NAME_CHECKED(UContextualAnimSceneAsset, OverridePreviewData))
+	{
+		SetDefaultMode();
+	}
 }
 
 bool FContextualAnimViewModel::MatchesContext(const FTransactionContext& InContext, const TArray<TPair<UObject*, FTransactionObjectEvent>>& TransactionObjectContexts) const
