Contextual Anim: Added missing includes to fix NonUnity build

diff --git a/Engine/Plugins/Experimental/Animation/ContextualAnimation/Source/ContextualAnimationEditor/Private/ContextualAnimViewModel.cpp b/Engine/Plugins/Experimental/Animation/ContextualAnimation/Source/ContextualAnimationEditor/Private/ContextualAnimViewModel.cpp
--- a/Engine/Plugins/Experimental/Animation/ContextualAnimation/Source/ContextualAnimationEditor/Private/ContextualAnimViewModel.cpp
+++ b/Engine/Plugins/Experimental/Animation/ContextualAnimation/Source/ContextualAnimationEditor/Private/ContextualAnimViewModel.cpp
@@ -37,6 +37,8 @@
 #include "ContextualAnimSelectionCriterion.h"
 #include "ScopedTransaction.h"
 #include "UObject/Package.h"
+#include "Engine/StaticMesh.h"
+#include "Components/StaticMeshComponent.h"
 
 #define LOCTEXT_NAMESPACE "ContextualAnimViewModel"
 
