Contextual Anim: Added random weight to each set. Will be used by the selection mechanism to 'break the tie' when multiple Sets can be selected

diff --git a/Engine/Plugins/Experimental/Animation/ContextualAnimation/Source/ContextualAnimation/Public/ContextualAnimSceneAsset.h b/Engine/Plugins/Experimental/Animation/ContextualAnimation/Source/ContextualAnimation/Public/ContextualAnimSceneAsset.h
--- a/Engine/Plugins/Experimental/Animation/ContextualAnimation/Source/ContextualAnimation/Public/ContextualAnimSceneAsset.h
+++ b/Engine/Plugins/Experimental/Animation/ContextualAnimation/Source/ContextualAnimation/Public/ContextualAnimSceneAsset.h
@@ -31,17 +31,27 @@ public:
 	FORCEINLINE int32 GetNumRoles() const { return Roles.Num(); }
 };
 
-/** List of AnimTrack for each role */
+/**
+ * Contains AnimTracks for each role in the interaction.
+ * Example: An specific set for a interaction with a car would have two tracks, one with the animation for the character and another one with the animation for the car.
+ * It is common to have variations of the same action with different animations. We could have one AnimSet with the animations for getting into the car from the driver side and another for getting into the car from the passenger side.
+*/
 USTRUCT(BlueprintType)
 struct CONTEXTUALANIMATION_API FContextualAnimSet
 {
 	GENERATED_BODY()
 
+	/** List of tracks with animation (and relevant data specific to that animation) for each role */
 	UPROPERTY(EditAnywhere, Category = "Defaults")
 	TArray<FContextualAnimTrack> Tracks;
 
+	/** Scene pivots for this set. Generated off line based on the AnimSetPivotDefinitions for the section this Set belongs to */
 	UPROPERTY(EditAnywhere, Category = "Defaults")
 	TArray<FTransform> ScenePivots;
+
+	/** Used by the selection mechanism to 'break the tie' when multiple Sets can be selected */
+	UPROPERTY(EditAnywhere, Category = "Defaults", meta = (ClampMin = "0", UIMin = "0", ClampMax = "1", UIMax = "1"))
+	float RandomWeight = 1.f;
 };
 
 /** Named container with one or more ContextualAnimSet */
diff --git a/Engine/Plugins/Experimental/Animation/ContextualAnimation/Source/ContextualAnimationEditor/Private/ContextualAnimViewModel.cpp b/Engine/Plugins/Experimental/Animation/ContextualAnimation/Source/ContextualAnimationEditor/Private/ContextualAnimViewModel.cpp
--- a/Engine/Plugins/Experimental/Animation/ContextualAnimation/Source/ContextualAnimationEditor/Private/ContextualAnimViewModel.cpp
+++ b/Engine/Plugins/Experimental/Animation/ContextualAnimation/Source/ContextualAnimationEditor/Private/ContextualAnimViewModel.cpp
@@ -479,6 +479,7 @@ void FContextualAnimViewModel::AddNewAnimSet(const FContextualAnimNewAnimSetPara
 		AnimTrack.Animation = Data.Animation;
 		AnimTrack.bRequireFlyingMode = Data.bRequiresFlyingMode;
 		AnimSet.Tracks.Add(AnimTrack);
+		AnimSet.RandomWeight = Params.RandomWeight;
 	}
 
 	int32 SectionIdx = INDEX_NONE;
diff --git a/Engine/Plugins/Experimental/Animation/ContextualAnimation/Source/ContextualAnimationEditor/Public/ContextualAnimEditorTypes.h b/Engine/Plugins/Experimental/Animation/ContextualAnimation/Source/ContextualAnimationEditor/Public/ContextualAnimEditorTypes.h
--- a/Engine/Plugins/Experimental/Animation/ContextualAnimation/Source/ContextualAnimationEditor/Public/ContextualAnimEditorTypes.h
+++ b/Engine/Plugins/Experimental/Animation/ContextualAnimation/Source/ContextualAnimationEditor/Public/ContextualAnimEditorTypes.h
@@ -88,4 +88,7 @@ struct FContextualAnimNewAnimSetParams
 
 	UPROPERTY(EditAnywhere, EditFixedSize, Category = "Settings", meta = (TitleProperty = "RoleName"))
 	TArray<FContextualAnimNewAnimSetData> Data;
+
+	UPROPERTY(EditAnywhere, Category = "Settings", meta = (ClampMin = "0", UIMin = "0", ClampMax = "1", UIMax = "1"))
+	float RandomWeight = 1.f;
 };
\ No newline at end of file
