Contextual Anim: Fix for crash that could happen in the editor when stopping Simulate mode

diff --git a/Engine/Plugins/Experimental/Animation/ContextualAnimation/Source/ContextualAnimation/Private/ContextualAnimSceneInstance.cpp b/Engine/Plugins/Experimental/Animation/ContextualAnimation/Source/ContextualAnimation/Private/ContextualAnimSceneInstance.cpp
--- a/Engine/Plugins/Experimental/Animation/ContextualAnimation/Source/ContextualAnimation/Private/ContextualAnimSceneInstance.cpp
+++ b/Engine/Plugins/Experimental/Animation/ContextualAnimation/Source/ContextualAnimation/Private/ContextualAnimSceneInstance.cpp
@@ -274,47 +274,50 @@ bool UContextualAnimSceneInstance::ForceTransitionToSection(const int32 SectionI
 
 void UContextualAnimSceneInstance::OnMontageBlendingOut(UAnimMontage* Montage, bool bInterrupted)
 {
-	UE_LOG(LogContextualAnim, Log, TEXT("UContextualAnimSceneInstance::OnMontageBlendingOut Montage: %s"), *GetNameSafe(Montage));
+	UE_LOG(LogContextualAnim, Log, TEXT("UContextualAnimSceneInstance::OnMontageBlendingOut Montage: %s bInterrupted: %d Bindings.IsValid: %d"), 
+		*GetNameSafe(Montage), bInterrupted, Bindings.IsValid());
 
-	for (auto& Binding : Bindings)
+	for (const FContextualAnimSceneBinding& Binding : Bindings)
 	{
 		const FContextualAnimTrack& AnimTrack = Bindings.GetAnimTrackFromBinding(Binding);
 		if (AnimTrack.Animation == Montage)
 		{
-			AActor* Actor = Binding.GetActor();
-			if (UAnimInstance* AnimInstance = Binding.GetAnimInstance())
+			if(AActor* Actor = Binding.GetActor())
 			{
-				AnimInstance->OnPlayMontageNotifyBegin.RemoveDynamic(this, &UContextualAnimSceneInstance::OnNotifyBeginReceived);
-				AnimInstance->OnPlayMontageNotifyEnd.RemoveDynamic(this, &UContextualAnimSceneInstance::OnNotifyEndReceived);
-				AnimInstance->OnMontageBlendingOut.RemoveDynamic(this, &UContextualAnimSceneInstance::OnMontageBlendingOut);
-
-				if (AnimTrack.bRequireFlyingMode)
+				if (UAnimInstance* AnimInstance = Binding.GetAnimInstance())
 				{
-					if (UCharacterMovementComponent* CharacterMovementComp = Actor->FindComponentByClass<UCharacterMovementComponent>())
+					AnimInstance->OnPlayMontageNotifyBegin.RemoveDynamic(this, &UContextualAnimSceneInstance::OnNotifyBeginReceived);
+					AnimInstance->OnPlayMontageNotifyEnd.RemoveDynamic(this, &UContextualAnimSceneInstance::OnNotifyEndReceived);
+					AnimInstance->OnMontageBlendingOut.RemoveDynamic(this, &UContextualAnimSceneInstance::OnMontageBlendingOut);
+
+					if (AnimTrack.bRequireFlyingMode)
 					{
-						CharacterMovementComp->SetMovementMode(MOVE_Walking);
+						if (UCharacterMovementComponent* CharacterMovementComp = Actor->FindComponentByClass<UCharacterMovementComponent>())
+						{
+							CharacterMovementComp->SetMovementMode(MOVE_Walking);
+						}
 					}
 				}
-			}
 
-			if (SceneAsset->GetDisableCollisionBetweenActors())
-			{
-				SetIgnoreCollisionWithOtherActors(Binding.GetActor(), false);
-			}
+				if (SceneAsset->GetDisableCollisionBetweenActors())
+				{
+					SetIgnoreCollisionWithOtherActors(Actor, false);
+				}
 
-			if (UContextualAnimSceneActorComponent* SceneActorComp = Binding.GetSceneActorComponent())
-			{
-				SceneActorComp->OnLeftScene();
-			}
+				if (UContextualAnimSceneActorComponent* SceneActorComp = Binding.GetSceneActorComponent())
+				{
+					SceneActorComp->OnLeftScene();
+				}
 
-			OnActorLeft.Broadcast(this, Actor);
+				OnActorLeft.Broadcast(this, Actor);
+			}
 
 			break;
 		}
 	}
 
 	bool bShouldEnd = true;
-	for (auto& Binding : Bindings)
+	for (const FContextualAnimSceneBinding& Binding : Bindings)
 	{
 		if (UAnimInstance* AnimInstance = Binding.GetAnimInstance())
 		{
