Contextual Anim Editor: Disable timeline, toolbar buttons and details panel when Simulate mode is active

diff --git a/Engine/Plugins/Experimental/Animation/ContextualAnimation/Source/ContextualAnimationEditor/Private/ContextualAnimAssetEditorToolkit.cpp b/Engine/Plugins/Experimental/Animation/ContextualAnimation/Source/ContextualAnimationEditor/Private/ContextualAnimAssetEditorToolkit.cpp
--- a/Engine/Plugins/Experimental/Animation/ContextualAnimation/Source/ContextualAnimationEditor/Private/ContextualAnimAssetEditorToolkit.cpp
+++ b/Engine/Plugins/Experimental/Animation/ContextualAnimation/Source/ContextualAnimationEditor/Private/ContextualAnimAssetEditorToolkit.cpp
@@ -96,6 +96,7 @@ void FContextualAnimAssetEditorToolkit::InitAssetEditor(const EToolkitMode::Type
 	EditingAssetWidget = PropertyModule.CreateDetailView(Args);
 	EditingAssetWidget->SetObject(SceneAsset);
 	EditingAssetWidget->OnFinishedChangingProperties().AddSP(this, &FContextualAnimAssetEditorToolkit::OnFinishedChangingProperties);
+	EditingAssetWidget->SetIsPropertyEditingEnabledDelegate(FIsPropertyEditingEnabled::CreateSP(this, &FContextualAnimAssetEditorToolkit::CanMakeEdits));
 
 	// Define Editor Layout
 	const TSharedRef<FTabManager::FLayout> StandaloneDefaultLayout = FTabManager::NewLayout("Standalone_ContextualAnimAnimEditor_Layout_v0.10")
@@ -156,11 +157,13 @@ void FContextualAnimAssetEditorToolkit::BindCommands()
 	ToolkitCommands->MapAction(
 		Commands.ResetPreviewScene,
 		FExecuteAction::CreateSP(this, &FContextualAnimAssetEditorToolkit::ResetPreviewScene),
+		FCanExecuteAction::CreateSP(this, &FContextualAnimAssetEditorToolkit::CanMakeEdits),
 		EUIActionRepeatMode::RepeatDisabled);
 
 	ToolkitCommands->MapAction(
 		Commands.NewAnimSet,
 		FExecuteAction::CreateSP(this, &FContextualAnimAssetEditorToolkit::ShowNewAnimSetDialog),
+		FCanExecuteAction::CreateSP(this, &FContextualAnimAssetEditorToolkit::CanMakeEdits),
 		EUIActionRepeatMode::RepeatDisabled);
 
 	ToolkitCommands->MapAction(
@@ -183,6 +186,11 @@ bool FContextualAnimAssetEditorToolkit::IsSimulateModeActive() const
 	return (ViewModel.IsValid() && !ViewModel->IsSimulateModeInactive());
 }
 
+bool FContextualAnimAssetEditorToolkit::CanMakeEdits() const
+{
+	return IsSimulateModeActive() == false;
+}
+
 void FContextualAnimAssetEditorToolkit::ExtendToolbar()
 {
 	TSharedPtr<FExtender> ToolbarExtender = MakeShareable(new FExtender);
@@ -208,7 +216,11 @@ void FContextualAnimAssetEditorToolkit::FillToolbar(FToolBarBuilder& ToolbarBuil
 	);
 
 	ToolbarBuilder.AddComboButton(
-		FUIAction(),
+		FUIAction(
+			FExecuteAction(),
+			FCanExecuteAction::CreateSP(this, &FContextualAnimAssetEditorToolkit::CanMakeEdits),
+			FIsActionChecked()
+		),
 		FOnGetContent::CreateSP(this, &FContextualAnimAssetEditorToolkit::BuildSectionsMenu),
 		LOCTEXT("Sections_Label", "Sections"),
 		FText::GetEmpty(),
diff --git a/Engine/Plugins/Experimental/Animation/ContextualAnimation/Source/ContextualAnimationEditor/Private/ContextualAnimViewModel.cpp b/Engine/Plugins/Experimental/Animation/ContextualAnimation/Source/ContextualAnimationEditor/Private/ContextualAnimViewModel.cpp
--- a/Engine/Plugins/Experimental/Animation/ContextualAnimation/Source/ContextualAnimationEditor/Private/ContextualAnimViewModel.cpp
+++ b/Engine/Plugins/Experimental/Animation/ContextualAnimation/Source/ContextualAnimationEditor/Private/ContextualAnimViewModel.cpp
@@ -546,6 +546,8 @@ void FContextualAnimViewModel::ToggleSimulateMode()
 {
 	if (SimulateModeState == ESimulateModeState::Inactive)
 	{
+		MovieScene->SetReadOnly(true);
+
 		SimulateModeState = ESimulateModeState::Paused;
 
 		if (SceneInstance.IsValid())
@@ -555,6 +557,8 @@ void FContextualAnimViewModel::ToggleSimulateMode()
 	}
 	else
 	{
+		MovieScene->SetReadOnly(false);
+
 		SimulateModeState = ESimulateModeState::Inactive;
 
 		if (SceneInstance.IsValid())
@@ -721,6 +725,14 @@ void FContextualAnimViewModel::SequencerDataChanged(EMovieSceneDataChangeType Da
 {
 	UE_LOG(LogContextualAnim, Log, TEXT("FContextualAnimViewModel::OnMovieSceneDataChanged DataChangeType: %d"), (int32)DataChangeType);
 
+	// Stop simulate mode if the user clicks the Unlock button on the Sequencer Panel while Simulate mode is active
+	// This is a temp solution, in the future we will hide that button along with the other Sequencer related buttons
+	if (DataChangeType == EMovieSceneDataChangeType::Unknown && SimulateModeState != ESimulateModeState::Inactive && Sequencer->IsReadOnly() == false)
+	{
+		ToggleSimulateMode();
+		return;
+	}
+
 	if(!SceneInstance.IsValid())
 	{
 		return;
diff --git a/Engine/Plugins/Experimental/Animation/ContextualAnimation/Source/ContextualAnimationEditor/Public/ContextualAnimAssetEditorToolkit.h b/Engine/Plugins/Experimental/Animation/ContextualAnimation/Source/ContextualAnimationEditor/Public/ContextualAnimAssetEditorToolkit.h
--- a/Engine/Plugins/Experimental/Animation/ContextualAnimation/Source/ContextualAnimationEditor/Public/ContextualAnimAssetEditorToolkit.h
+++ b/Engine/Plugins/Experimental/Animation/ContextualAnimation/Source/ContextualAnimationEditor/Public/ContextualAnimAssetEditorToolkit.h
@@ -42,6 +42,8 @@ public:
 
 private:
 	
+	bool CanMakeEdits() const;
+
 	TSharedRef<SDockTab> SpawnTab_Viewport(const FSpawnTabArgs& Args);
 	TSharedRef<SDockTab> SpawnTab_AssetDetails(const FSpawnTabArgs& Args);
 	TSharedRef<SDockTab> SpawnTab_Timeline(const FSpawnTabArgs& Args);
