Contextual Anim Editor: Added ScenActorComp to preview actors

diff --git a/Engine/Plugins/Experimental/Animation/ContextualAnimation/Source/ContextualAnimationEditor/Private/ContextualAnimViewModel.cpp b/Engine/Plugins/Experimental/Animation/ContextualAnimation/Source/ContextualAnimationEditor/Private/ContextualAnimViewModel.cpp
--- a/Engine/Plugins/Experimental/Animation/ContextualAnimation/Source/ContextualAnimationEditor/Private/ContextualAnimViewModel.cpp
+++ b/Engine/Plugins/Experimental/Animation/ContextualAnimation/Source/ContextualAnimationEditor/Private/ContextualAnimViewModel.cpp
@@ -24,6 +24,7 @@
 #include "ContextualAnimSceneAsset.h"
 #include "ContextualAnimManager.h"
 #include "ContextualAnimSceneInstance.h"
+#include "ContextualAnimSceneActorComponent.h"
 #include "AnimNotifyState_IKWindow.h"
 #include "PropertyEditorModule.h"
 #include "Modules/ModuleManager.h"
@@ -277,6 +278,12 @@ AActor* FContextualAnimViewModel::SpawnPreviewActor(const FContextualAnimTrack&
 				}
 			}
 		}
+
+		check(PreviewActor);
+		UContextualAnimSceneActorComponent* SceneActorComp = NewObject<UContextualAnimSceneActorComponent>(PreviewActor);
+		SceneActorComp->RegisterComponentWithWorld(GetWorld());
+		SceneActorComp->InitializeComponent();
+
 	}
 
 	return PreviewActor;
