Contextual Anim Editor: Starting the 'simulation' of the interaction can be now done with the Play button instead of Enter key

diff --git a/Engine/Plugins/Experimental/Animation/ContextualAnimation/Source/ContextualAnimationEditor/Private/ContextualAnimEdMode.cpp b/Engine/Plugins/Experimental/Animation/ContextualAnimation/Source/ContextualAnimationEditor/Private/ContextualAnimEdMode.cpp
--- a/Engine/Plugins/Experimental/Animation/ContextualAnimation/Source/ContextualAnimationEditor/Private/ContextualAnimEdMode.cpp
+++ b/Engine/Plugins/Experimental/Animation/ContextualAnimation/Source/ContextualAnimationEditor/Private/ContextualAnimEdMode.cpp
@@ -327,15 +327,6 @@ bool FContextualAnimEdMode::InputDelta(FEditorViewportClient* InViewportClient,
 
 bool FContextualAnimEdMode::InputKey(FEditorViewportClient* ViewportClient, FViewport* Viewport, FKey Key, EInputEvent Event)
 {
-	if (ViewModel && ViewModel->IsSimulateModePaused())
-	{
-		if (Key == EKeys::Enter && Event == IE_Released)
-		{
-			ViewModel->StartSimulation();
-			return true;
-		}
-	}
-
 	return FEdMode::InputKey(ViewportClient, Viewport, Key, Event);
 }
 
diff --git a/Engine/Plugins/Experimental/Animation/ContextualAnimation/Source/ContextualAnimationEditor/Private/ContextualAnimViewModel.cpp b/Engine/Plugins/Experimental/Animation/ContextualAnimation/Source/ContextualAnimationEditor/Private/ContextualAnimViewModel.cpp
--- a/Engine/Plugins/Experimental/Animation/ContextualAnimation/Source/ContextualAnimationEditor/Private/ContextualAnimViewModel.cpp
+++ b/Engine/Plugins/Experimental/Animation/ContextualAnimation/Source/ContextualAnimationEditor/Private/ContextualAnimViewModel.cpp
@@ -57,6 +57,8 @@ FContextualAnimViewModel::~FContextualAnimViewModel()
 	{
 		Sequencer->OnMovieSceneDataChanged().RemoveAll(this);
 		Sequencer->OnGlobalTimeChanged().RemoveAll(this);
+		Sequencer->OnPlayEvent().RemoveAll(this);
+		Sequencer->OnStopEvent().RemoveAll(this);
 		Sequencer.Reset();
 	}
 }
@@ -114,6 +116,8 @@ void FContextualAnimViewModel::CreateSequencer()
 	Sequencer = SequencerModule.CreateSequencer(SequencerInitParams);
 	Sequencer->OnMovieSceneDataChanged().AddRaw(this, &FContextualAnimViewModel::SequencerDataChanged);
 	Sequencer->OnGlobalTimeChanged().AddRaw(this, &FContextualAnimViewModel::SequencerTimeChanged);
+	Sequencer->OnPlayEvent().AddRaw(this, &FContextualAnimViewModel::SequencerPlayEvent);
+	Sequencer->OnStopEvent().AddRaw(this, &FContextualAnimViewModel::SequencerStopEvent);
 	Sequencer->SetPlaybackStatus(EMovieScenePlayerStatus::Stopped);
 }
 
@@ -680,6 +684,22 @@ UContextualAnimMovieSceneTrack* FContextualAnimViewModel::FindTrackByRole(const
 	return nullptr;
 }
 
+void FContextualAnimViewModel::SequencerPlayEvent()
+{
+	if (SimulateModeState == ESimulateModeState::Paused)
+	{
+		StartSimulation();
+	}
+}
+
+void FContextualAnimViewModel::SequencerStopEvent()
+{
+	if (SimulateModeState != ESimulateModeState::Inactive)
+	{
+		ToggleSimulateMode();
+	}
+}
+
 void FContextualAnimViewModel::SequencerTimeChanged()
 {
 	EMovieScenePlayerStatus::Type CurrentStatus = Sequencer->GetPlaybackStatus();
diff --git a/Engine/Plugins/Experimental/Animation/ContextualAnimation/Source/ContextualAnimationEditor/Public/ContextualAnimViewModel.h b/Engine/Plugins/Experimental/Animation/ContextualAnimation/Source/ContextualAnimationEditor/Public/ContextualAnimViewModel.h
--- a/Engine/Plugins/Experimental/Animation/ContextualAnimation/Source/ContextualAnimationEditor/Public/ContextualAnimViewModel.h
+++ b/Engine/Plugins/Experimental/Animation/ContextualAnimation/Source/ContextualAnimationEditor/Public/ContextualAnimViewModel.h
@@ -211,6 +211,10 @@ private:
 
 	void SequencerDataChanged(EMovieSceneDataChangeType DataChangeType);
 
+	void SequencerPlayEvent();
+
+	void SequencerStopEvent();
+
 	void OnAnimNotifyChanged(UAnimSequenceBase* Animation);
 
 	void CreateSequencer();
