Contextual Anim Editor: Details Tab now only shows relevant info from the active set and the selected anim track instead of the full raw details panel

diff --git a/Engine/Plugins/Experimental/Animation/ContextualAnimation/Source/ContextualAnimation/Public/ContextualAnimSceneAsset.h b/Engine/Plugins/Experimental/Animation/ContextualAnimation/Source/ContextualAnimation/Public/ContextualAnimSceneAsset.h
--- a/Engine/Plugins/Experimental/Animation/ContextualAnimation/Source/ContextualAnimation/Public/ContextualAnimSceneAsset.h
+++ b/Engine/Plugins/Experimental/Animation/ContextualAnimation/Source/ContextualAnimation/Public/ContextualAnimSceneAsset.h
@@ -275,13 +275,13 @@ public:
 
 protected:
 
-	UPROPERTY(EditAnywhere, Category = "Defaults")
+	UPROPERTY(EditAnywhere, Category = "Settings")
 	TObjectPtr<UContextualAnimRolesAsset> RolesAsset;
 
-	UPROPERTY(EditAnywhere, Category = "Defaults", meta = (GetOptions = "GetRoles"))
+	UPROPERTY(EditAnywhere, Category = "Settings", meta = (GetOptions = "GetRoles"))
 	FName PrimaryRole = NAME_None;
 
-	UPROPERTY(EditAnywhere, Category = "Defaults")
+	UPROPERTY(EditAnywhere, Category = "Settings", meta = (TitleProperty = "Role"))
 	TArray<FContextualAnimActorPreviewData> OverridePreviewData;
 
 	UPROPERTY(EditAnywhere, Category = "Defaults")
diff --git a/Engine/Plugins/Experimental/Animation/ContextualAnimation/Source/ContextualAnimation/Public/ContextualAnimTypes.h b/Engine/Plugins/Experimental/Animation/ContextualAnimation/Source/ContextualAnimation/Public/ContextualAnimTypes.h
--- a/Engine/Plugins/Experimental/Animation/ContextualAnimation/Source/ContextualAnimation/Public/ContextualAnimTypes.h
+++ b/Engine/Plugins/Experimental/Animation/ContextualAnimation/Source/ContextualAnimation/Public/ContextualAnimTypes.h
@@ -89,16 +89,16 @@ struct CONTEXTUALANIMATION_API FContextualAnimTrack
 	UPROPERTY(EditAnywhere, BlueprintReadWrite, Category = "Defaults")
 	FTransform MeshToScene;
 
-	UPROPERTY(EditAnywhere, BlueprintReadOnly, Category = "Defaults", meta = (GetOptions = "GetRoles"))
+	UPROPERTY(VisibleAnywhere, BlueprintReadOnly, Category = "Defaults")
 	FName Role = NAME_None;
 
-	UPROPERTY(VisibleAnywhere, BlueprintReadOnly, Category = "Defaults")
+	UPROPERTY(BlueprintReadOnly, Category = "Defaults")
 	int32 SectionIdx = INDEX_NONE;
 
-	UPROPERTY(VisibleAnywhere, BlueprintReadOnly, Category = "Defaults")
+	UPROPERTY(BlueprintReadOnly, Category = "Defaults")
 	int32 AnimSetIdx = INDEX_NONE;
 
-	UPROPERTY(VisibleAnywhere, BlueprintReadOnly, Category = "Defaults")
+	UPROPERTY(BlueprintReadOnly, Category = "Defaults")
 	int32 AnimTrackIdx = INDEX_NONE;
 
 	// DEPRECATED: Will go away soon
diff --git a/Engine/Plugins/Experimental/Animation/ContextualAnimation/Source/ContextualAnimationEditor/Private/ContextualAnimAssetEditorToolkit.cpp b/Engine/Plugins/Experimental/Animation/ContextualAnimation/Source/ContextualAnimationEditor/Private/ContextualAnimAssetEditorToolkit.cpp
--- a/Engine/Plugins/Experimental/Animation/ContextualAnimation/Source/ContextualAnimationEditor/Private/ContextualAnimAssetEditorToolkit.cpp
+++ b/Engine/Plugins/Experimental/Animation/ContextualAnimation/Source/ContextualAnimationEditor/Private/ContextualAnimAssetEditorToolkit.cpp
@@ -49,6 +49,8 @@ FContextualAnimAssetEditorToolkit::FContextualAnimAssetEditorToolkit()
 
 FContextualAnimAssetEditorToolkit::~FContextualAnimAssetEditorToolkit()
 {
+	check(ViewModel.IsValid());
+	ViewModel->Shutdown();
 }
 
 UContextualAnimSceneAsset* FContextualAnimAssetEditorToolkit::GetSceneAsset() const
@@ -86,20 +88,8 @@ void FContextualAnimAssetEditorToolkit::InitAssetEditor(const EToolkitMode::Type
 	ViewModel = MakeShared<FContextualAnimViewModel>();
 	ViewModel->Initialize(SceneAsset, PreviewScene.ToSharedRef());
 
-	// Create Asset Details widget
-	FPropertyEditorModule& PropertyModule = FModuleManager::LoadModuleChecked<FPropertyEditorModule>("PropertyEditor");
-
-	FDetailsViewArgs Args;
-	Args.bHideSelectionTip = true;
-	Args.NotifyHook = this;
-
-	EditingAssetWidget = PropertyModule.CreateDetailView(Args);
-	EditingAssetWidget->SetObject(SceneAsset);
-	EditingAssetWidget->OnFinishedChangingProperties().AddSP(this, &FContextualAnimAssetEditorToolkit::OnFinishedChangingProperties);
-	EditingAssetWidget->SetIsPropertyEditingEnabledDelegate(FIsPropertyEditingEnabled::CreateSP(this, &FContextualAnimAssetEditorToolkit::CanMakeEdits));
-
 	// Define Editor Layout
-	const TSharedRef<FTabManager::FLayout> StandaloneDefaultLayout = FTabManager::NewLayout("Standalone_ContextualAnimAnimEditor_Layout_v0.10")
+	const TSharedRef<FTabManager::FLayout> StandaloneDefaultLayout = FTabManager::NewLayout("Standalone_ContextualAnimAnimEditor_Layout_v0.13")
 		->AddArea
 		(
 			FTabManager::NewPrimaryArea()->SetOrientation(Orient_Vertical)
@@ -188,7 +178,7 @@ bool FContextualAnimAssetEditorToolkit::IsSimulateModeActive() const
 
 bool FContextualAnimAssetEditorToolkit::CanMakeEdits() const
 {
-	return IsSimulateModeActive() == false;
+	return (ViewModel.IsValid() && ViewModel->CanMakeEdits());
 }
 
 void FContextualAnimAssetEditorToolkit::ExtendToolbar()
@@ -379,7 +369,7 @@ TSharedRef<SDockTab> FContextualAnimAssetEditorToolkit::SpawnTab_AssetDetails(co
 	return SNew(SDockTab)
 		.Label(LOCTEXT("AssetDetails_Title", "Asset Details"))
 		[
-			EditingAssetWidget.ToSharedRef()
+			ViewModel->GetDetailsView().ToSharedRef()
 		];
 }
 
@@ -413,10 +403,5 @@ TSharedRef<SDockTab> FContextualAnimAssetEditorToolkit::SpawnTab_PreviewSettings
 	return SpawnedTab;
 }
 
-void FContextualAnimAssetEditorToolkit::OnFinishedChangingProperties(const FPropertyChangedEvent& PropertyChangedEvent)
-{
-	ViewModel->OnFinishedChangingProperties(PropertyChangedEvent);
-}
-
 #undef LOCTEXT_NAMESPACE
 
diff --git a/Engine/Plugins/Experimental/Animation/ContextualAnimation/Source/ContextualAnimationEditor/Private/ContextualAnimViewModel.cpp b/Engine/Plugins/Experimental/Animation/ContextualAnimation/Source/ContextualAnimationEditor/Private/ContextualAnimViewModel.cpp
--- a/Engine/Plugins/Experimental/Animation/ContextualAnimation/Source/ContextualAnimationEditor/Private/ContextualAnimViewModel.cpp
+++ b/Engine/Plugins/Experimental/Animation/ContextualAnimation/Source/ContextualAnimationEditor/Private/ContextualAnimViewModel.cpp
@@ -40,9 +40,12 @@
 #include "UObject/Package.h"
 #include "Engine/StaticMesh.h"
 #include "Components/StaticMeshComponent.h"
+#include "DetailCustomizations/ContextualAnimSceneAssetDetailCustom.h"
 
 #define LOCTEXT_NAMESPACE "ContextualAnimViewModel"
 
+static const FName PreviewActorTag = FName(TEXT("__PreviewActor__"));
+
 FContextualAnimViewModel::FContextualAnimViewModel()
 	: SceneAsset(nullptr)
 	, MovieSceneSequence(nullptr)
@@ -53,14 +56,7 @@ FContextualAnimViewModel::FContextualAnimViewModel()
 
 FContextualAnimViewModel::~FContextualAnimViewModel()
 {
-	if (Sequencer.IsValid())
-	{
-		Sequencer->OnMovieSceneDataChanged().RemoveAll(this);
-		Sequencer->OnGlobalTimeChanged().RemoveAll(this);
-		Sequencer->OnPlayEvent().RemoveAll(this);
-		Sequencer->OnStopEvent().RemoveAll(this);
-		Sequencer.Reset();
-	}
+	check(!bInitialized)
 }
 
 void FContextualAnimViewModel::AddReferencedObjects(FReferenceCollector& Collector)
@@ -78,6 +74,8 @@ TSharedPtr<ISequencer> FContextualAnimViewModel::GetSequencer()
 
 void FContextualAnimViewModel::Initialize(UContextualAnimSceneAsset* InSceneAsset, const TSharedRef<FContextualAnimPreviewScene>& InPreviewScene)
 {
+	check(!bInitialized);
+
 	SceneAsset = InSceneAsset;
 	PreviewScenePtr = InPreviewScene;
 
@@ -85,7 +83,47 @@ void FContextualAnimViewModel::Initialize(UContextualAnimSceneAsset* InSceneAsse
 
 	CreateSequencer();
 
+	CreateDetailsView();
+
 	SetDefaultMode();
+
+	bInitialized = true;
+}
+
+void FContextualAnimViewModel::Shutdown()
+{
+	check(DetailsView.IsValid())
+	DetailsView->UnregisterInstancedCustomPropertyLayout(UContextualAnimSceneAsset::StaticClass());
+	DetailsView.Reset();
+
+	check(Sequencer.IsValid());
+	Sequencer->OnMovieSceneDataChanged().RemoveAll(this);
+	Sequencer->OnGlobalTimeChanged().RemoveAll(this);
+	Sequencer->OnPlayEvent().RemoveAll(this);
+	Sequencer->OnStopEvent().RemoveAll(this);
+	Sequencer.Reset();
+
+	bInitialized = false;
+}
+
+bool FContextualAnimViewModel::CanMakeEdits() const
+{
+	return IsSimulateModeInactive();
+}
+
+void FContextualAnimViewModel::CreateDetailsView()
+{
+	FPropertyEditorModule& PropertyModule = FModuleManager::LoadModuleChecked<FPropertyEditorModule>("PropertyEditor");
+
+	FDetailsViewArgs Args;
+	Args.bHideSelectionTip = true;
+
+	DetailsView = PropertyModule.CreateDetailView(Args);
+	DetailsView->SetObject(SceneAsset);
+	DetailsView->OnFinishedChangingProperties().AddSP(this, &FContextualAnimViewModel::OnFinishedChangingProperties);
+	DetailsView->SetIsPropertyEditingEnabledDelegate(FIsPropertyEditingEnabled::CreateSP(this, &FContextualAnimViewModel::CanMakeEdits));
+	DetailsView->RegisterInstancedCustomPropertyLayout(UContextualAnimSceneAsset::StaticClass(), 
+		FOnGetDetailCustomizationInstance::CreateStatic(&FContextualAnimSceneAssetDetailCustom::MakeInstance, AsShared()));
 }
 
 void FContextualAnimViewModel::CreateSequencer()
@@ -135,6 +173,11 @@ void FContextualAnimViewModel::SetActiveSection(int32 SectionIdx)
 	SetDefaultMode();
 }
 
+int32 FContextualAnimViewModel::GetActiveSection() const
+{
+	return ActiveSectionIdx;
+}
+
 void FContextualAnimViewModel::SetActiveAnimSetForSection(int32 SectionIdx, int32 AnimSetIdx)
 {
 	check(GetSceneAsset()->Sections.IsValidIndex(SectionIdx));
@@ -146,6 +189,12 @@ void FContextualAnimViewModel::SetActiveAnimSetForSection(int32 SectionIdx, int3
 	SetDefaultMode();
 }
 
+int32 FContextualAnimViewModel::GetActiveAnimSetForSection(int32 SectionIdx) const
+{
+	const int32* ActiveSetIdxPtr = ActiveAnimSetMap.Find(SectionIdx);
+	return ActiveSetIdxPtr ? *ActiveSetIdxPtr : INDEX_NONE;
+}
+
 AActor* FContextualAnimViewModel::SpawnPreviewActor(const FContextualAnimTrack& AnimTrack)
 {
 	AActor* PreviewActor = nullptr;
@@ -288,6 +337,7 @@ AActor* FContextualAnimViewModel::SpawnPreviewActor(const FContextualAnimTrack&
 		SceneActorComp->RegisterComponentWithWorld(GetWorld());
 		SceneActorComp->InitializeComponent();
 
+		PreviewActor->Tags.Add(PreviewActorTag);
 	}
 
 	return PreviewActor;
@@ -306,6 +356,8 @@ void FContextualAnimViewModel::ResetTimeline()
 
 void FContextualAnimViewModel::SetDefaultMode()
 {
+	ClearSelection();
+
 	if (EditingAnimation.IsValid())
 	{
 		EditingAnimation->UnregisterOnNotifyChanged(this);
@@ -429,9 +481,12 @@ void FContextualAnimViewModel::RefreshPreviewScene()
 		SceneInstance.Reset();
 	}
 
-	for (const FContextualAnimSceneBinding& Binding : SceneBindings)
+	// Remove preview actors from the preview scene
+	// @TODO: Investigate why GetActor from each binding in SceneBindings returns null after a Redo op
+	for (TActorIterator<AActor> It(GetWorld()); It; ++It)
 	{
-		if (AActor* Actor = Binding.GetActor())
+		AActor* Actor = *It;
+		if(Actor && Actor->ActorHasTag(PreviewActorTag))
 		{
 			Actor->Destroy();
 		}
@@ -525,6 +580,51 @@ void FContextualAnimViewModel::AddNewAnimSet(const FContextualAnimNewAnimSetPara
 	SetActiveAnimSetForSection(SectionIdx, AnimSetIdx);
 }
 
+void FContextualAnimViewModel::RemoveSection(int32 SectionIdx)
+{
+ 	check(SceneAsset->Sections.IsValidIndex(SectionIdx));
+
+	FScopedTransaction Transaction(LOCTEXT("RemoveSection", "Remove Section"));
+	SceneAsset->Modify();
+
+ 	SceneAsset->Sections.RemoveAt(SectionIdx);
+
+	if (SceneAsset->GetNumSections() > 0)
+	{
+		SetActiveAnimSetForSection(0, 0);
+	}
+	else
+	{
+		ActiveSectionIdx = INDEX_NONE;
+		ActiveAnimSetMap.Reset();
+	}
+
+	SetDefaultMode();
+}
+
+void FContextualAnimViewModel::RemoveAnimSet(int32 SectionIdx, int32 AnimSetIdx)
+{
+	check(SceneAsset->Sections.IsValidIndex(SectionIdx));
+	check(SceneAsset->Sections[SectionIdx].AnimSets.IsValidIndex(AnimSetIdx));
+
+	FScopedTransaction Transaction(LOCTEXT("RemoveAnimSet", "Remove Anim Set"));
+	SceneAsset->Modify();
+
+	SceneAsset->Sections[SectionIdx].AnimSets.RemoveAt(AnimSetIdx);
+
+	// If there are other sets in the current section, switch to the first one
+	if (SceneAsset->Sections[SectionIdx].GetNumAnimSets() > 0)
+	{
+		SetActiveAnimSetForSection(SectionIdx, 0);
+
+		SetDefaultMode();
+	}
+	else // Otherwise, remove the section too, there is no point in having a section without sets
+	{
+		RemoveSection(SectionIdx);
+	}
+}
+
 void FContextualAnimViewModel::AddNewIKTarget(const UContextualAnimNewIKTargetParams& Params)
 {
 	check(SceneAsset->Sections.IsValidIndex(Params.SectionIdx));
@@ -825,6 +925,11 @@ void FContextualAnimViewModel::AnimationModified(UAnimSequenceBase& Animation)
 	Animation.MarkPackageDirty();
 }
 
+void FContextualAnimViewModel::RefreshDetailsView()
+{
+	DetailsView->ForceRefresh();
+}
+
 void FContextualAnimViewModel::UpdateSelection(const AActor* SelectedActor)
 {
 	const FContextualAnimSceneBinding* Binding = SceneBindings.FindBindingByActor(SelectedActor);
@@ -852,11 +957,15 @@ void FContextualAnimViewModel::UpdateSelection(FName Role, int32 CriterionIdx, i
 	SelectionInfo.Role = Role;
 	SelectionInfo.Criterion.Key = CriterionIdx;
 	SelectionInfo.Criterion.Value = CriterionDataIdx;
+
+	RefreshDetailsView();
 }
 
 void FContextualAnimViewModel::ClearSelection()
 {
 	SelectionInfo.Reset();
+
+	RefreshDetailsView();
 }
 
 FContextualAnimSceneBinding* FContextualAnimViewModel::GetSelectedBinding() const
@@ -1119,14 +1228,17 @@ void FContextualAnimViewModel::OnFinishedChangingProperties(const FPropertyChang
 	const FName PropertyName = (PropertyChangedEvent.Property != nullptr) ? PropertyChangedEvent.Property->GetFName() : NAME_None;
 	const FName MemberPropertyName = (PropertyChangedEvent.MemberProperty != nullptr) ? PropertyChangedEvent.MemberProperty->GetFName() : NAME_None;
 
-	// Refresh preview scene if the Sections array or the AnimSets array inside a section changes
+	UE_LOG(LogContextualAnim, Verbose, TEXT("FContextualAnimViewModel::OnFinishedChangingProperties MemberPropertyName: %s PropertyName: %s"), *MemberPropertyName.ToString(), *PropertyName.ToString());
+
+	// Refresh preview scene if necessary
 	if (MemberPropertyName == GET_MEMBER_NAME_CHECKED(UContextualAnimSceneAsset, Sections) &&
 		(PropertyName == GET_MEMBER_NAME_CHECKED(UContextualAnimSceneAsset, Sections) ||
-			PropertyName == GET_MEMBER_NAME_CHECKED(FContextualAnimSceneSection, AnimSets)))
+		 PropertyName == GET_MEMBER_NAME_CHECKED(FContextualAnimSceneSection, AnimSets) || 
+		 PropertyName == GET_MEMBER_NAME_CHECKED(FContextualAnimTrack, Animation) ||
+		 PropertyName == GET_MEMBER_NAME_CHECKED(FContextualAnimTrack, MeshToScene)))
 	{
 		SetDefaultMode();
 	}
-	// Refresh preview scene if override preview data changes
 	else if (MemberPropertyName == GET_MEMBER_NAME_CHECKED(UContextualAnimSceneAsset, OverridePreviewData))
 	{
 		SetDefaultMode();
diff --git a/Engine/Plugins/Experimental/Animation/ContextualAnimation/Source/ContextualAnimationEditor/Private/DetailCustomizations/ContextualAnimSceneAssetDetailCustom.cpp b/Engine/Plugins/Experimental/Animation/ContextualAnimation/Source/ContextualAnimationEditor/Private/DetailCustomizations/ContextualAnimSceneAssetDetailCustom.cpp
new file mode 100644
--- /dev/null
+++ b/Engine/Plugins/Experimental/Animation/ContextualAnimation/Source/ContextualAnimationEditor/Private/DetailCustomizations/ContextualAnimSceneAssetDetailCustom.cpp
@@ -0,0 +1,136 @@
+// Copyright Epic Games, Inc. All Rights Reserved.
+
+#include "DetailCustomizations/ContextualAnimSceneAssetDetailCustom.h"
+#include "DetailLayoutBuilder.h"
+#include "DetailCategoryBuilder.h"
+#include "DetailWidgetRow.h"
+#include "PropertyEditorModule.h"
+#include "Widgets/Input/SButton.h"
+#include "Misc/MessageDialog.h"
+#include "ContextualAnimSceneAsset.h"
+#include "ContextualAnimViewModel.h"
+
+#define LOCTEXT_NAMESPACE "ContextualAnimSceneAssetDetailCustom"
+
+TSharedRef<IDetailCustomization> FContextualAnimSceneAssetDetailCustom::MakeInstance(TSharedRef<FContextualAnimViewModel> ViewModelRef)
+{
+	return MakeShared<FContextualAnimSceneAssetDetailCustom>(ViewModelRef);
+}
+
+FContextualAnimSceneAssetDetailCustom::FContextualAnimSceneAssetDetailCustom(TSharedRef<FContextualAnimViewModel> ViewModelRef)
+	: ViewModelPtr(ViewModelRef)
+{
+	
+}
+
+void FContextualAnimSceneAssetDetailCustom::CustomizeDetails(IDetailLayoutBuilder& DetailBuilder)
+{
+	check(ViewModelPtr.IsValid());
+	TSharedPtr<FContextualAnimViewModel> ViewModel = ViewModelPtr.Pin();
+
+	TArray<TWeakObjectPtr<UObject>> ObjectList;
+	DetailBuilder.GetObjectsBeingCustomized(ObjectList);
+	check(ObjectList.Num() > 0);
+
+	const UContextualAnimSceneAsset* SceneAsset = Cast<const UContextualAnimSceneAsset>(ObjectList[0].Get());
+	check(SceneAsset);
+
+	IDetailCategoryBuilder& Category = DetailBuilder.EditCategory(TEXT("Defaults"));
+	Category.SetCategoryVisibility(false);
+
+	const int32 SectionIdx = ViewModel->GetActiveSection();
+	if (SectionIdx != INDEX_NONE)
+	{
+		// Create new category to show editable data only for the currently active section in the editor
+		IDetailCategoryBuilder& CurrentSectionCategory = DetailBuilder.EditCategory(TEXT("Active Section"));
+
+		// Get IPropertyHandle for the array of Sections
+		TSharedPtr<IPropertyHandle> SectionsPropertyHandle = DetailBuilder.GetProperty(FName(TEXT("Sections")));
+		SectionsPropertyHandle->MarkHiddenByCustomization();
+
+		// Get IPropertyHandle for the current Section
+		TSharedRef<IPropertyHandle> SectionPropertyHandle = SectionsPropertyHandle->GetChildHandle(SectionIdx).ToSharedRef();
+
+		// Show editable properties in the custom category
+		CurrentSectionCategory.AddProperty(SectionPropertyHandle->GetChildHandle(FName(TEXT("Name"))).ToSharedRef());
+		CurrentSectionCategory.AddProperty(SectionPropertyHandle->GetChildHandle(FName(TEXT("RoleToIKTargetDefsMap"))).ToSharedRef());
+		CurrentSectionCategory.AddProperty(SectionPropertyHandle->GetChildHandle(FName(TEXT("AnimSetPivotDefinitions"))).ToSharedRef());
+
+		// Add button to remove the entire section
+		CurrentSectionCategory.AddCustomRow(FText::GetEmpty())
+			.ValueContent()
+			.VAlign(VAlign_Center)
+			.MaxDesiredWidth(250)
+			[
+				SNew(SButton).Text(LOCTEXT("RemoveSectionButtonLabel", "Remove Section"))
+				.OnClicked_Lambda([ViewModel, SectionIdx]()
+				{
+					if (FMessageDialog::Open(EAppMsgType::YesNo, LOCTEXT("RemoveSectionDialog", "Are you sure you want to remove this Section?")) == EAppReturnType::Yes)
+					{
+						ViewModel->RemoveSection(SectionIdx);
+					}
+
+					return FReply::Handled();
+				})
+			];
+
+		// Same as above but for the currently active AnimSet
+		const int32 AnimSetIdx = ViewModel->GetActiveAnimSetForSection(SectionIdx);
+		if (AnimSetIdx != INDEX_NONE)
+		{
+			// Create new category to show editable data only for the currently active AnimSet in the editor
+			IDetailCategoryBuilder& CurrentAnimSetCategory = DetailBuilder.EditCategory(TEXT("Active AnimSet"));
+
+			// Get IPropertyHandle for the array of AnimSet
+			TSharedPtr<IPropertyHandle> AnimSetArrayPropertyHandle = SectionPropertyHandle->GetChildHandle(FName(TEXT("AnimSets")));
+
+			// Get IPropertyHandle for the current AnimSet
+			TSharedRef<IPropertyHandle> AnimSetPropertyHandle = AnimSetArrayPropertyHandle->GetChildHandle(AnimSetIdx).ToSharedRef();
+
+			// Show general AnimSet properties in a custom category
+			CurrentAnimSetCategory.AddProperty(AnimSetPropertyHandle->GetChildHandle(FName(TEXT("RandomWeight"))).ToSharedRef());
+
+			// Add button to remove the entire set
+			CurrentAnimSetCategory.AddCustomRow(FText::GetEmpty())
+				.ValueContent()
+				.VAlign(VAlign_Center)
+				.MaxDesiredWidth(250)
+				[
+					SNew(SButton).Text(LOCTEXT("RemoveAnimSetButtonLabel", "Remove AnimSet"))
+					.OnClicked_Lambda([ViewModel, SectionIdx, AnimSetIdx]()
+					{
+						if (FMessageDialog::Open(EAppMsgType::YesNo, LOCTEXT("RemoveAnimSetDialog", "Are you sure you want to remove this AnimSet?")) == EAppReturnType::Yes)
+						{
+							ViewModel->RemoveAnimSet(SectionIdx, AnimSetIdx);
+						}
+						
+						return FReply::Handled();
+					})
+				];
+
+			// Same as above but for the selected AnimTrack
+			if (const FContextualAnimTrack* SelectedAnimTrack = ViewModel->GetSelectedAnimTrack())
+			{
+				// Create new category to show editable data only for the currently selected track
+				IDetailCategoryBuilder& AnimTrackCategory = DetailBuilder.EditCategory(TEXT("Selected AnimTrack"));
+
+				// Get IPropertyHandle for the array of AnimTracks
+				TSharedPtr<IPropertyHandle> TracksPropertyHandle = AnimSetPropertyHandle->GetChildHandle(FName(TEXT("Tracks")));
+
+				// Get IPropertyHandle for the selected AnimTrack
+				TSharedRef<IPropertyHandle> AnimTrackPropertyHandle = TracksPropertyHandle->GetChildHandle(SelectedAnimTrack->AnimTrackIdx).ToSharedRef();
+				AnimTrackPropertyHandle->MarkHiddenByCustomization();
+
+				// Show all the editable properties of the AnimTrack in a custom category
+				uint32 AnimTrackChildren = 0;
+				AnimTrackPropertyHandle->GetNumChildren(AnimTrackChildren);
+				for (uint32 Index = 0; Index < AnimTrackChildren; ++Index)
+				{
+					AnimTrackCategory.AddProperty(AnimTrackPropertyHandle->GetChildHandle(Index).ToSharedRef());
+				}
+			}
+		}
+	}
+}
+
+#undef LOCTEXT_NAMESPACE
\ No newline at end of file
diff --git a/Engine/Plugins/Experimental/Animation/ContextualAnimation/Source/ContextualAnimationEditor/Public/ContextualAnimAssetEditorToolkit.h b/Engine/Plugins/Experimental/Animation/ContextualAnimation/Source/ContextualAnimationEditor/Public/ContextualAnimAssetEditorToolkit.h
--- a/Engine/Plugins/Experimental/Animation/ContextualAnimation/Source/ContextualAnimationEditor/Public/ContextualAnimAssetEditorToolkit.h
+++ b/Engine/Plugins/Experimental/Animation/ContextualAnimation/Source/ContextualAnimationEditor/Public/ContextualAnimAssetEditorToolkit.h
@@ -15,7 +15,7 @@ class FContextualAnimPreviewScene;
 class FContextualAnimViewModel;
 class UContextualAnimPreviewManager;
 
-class FContextualAnimAssetEditorToolkit : public FAssetEditorToolkit, public FNotifyHook
+class FContextualAnimAssetEditorToolkit : public FAssetEditorToolkit
 {
 public:
 
@@ -58,14 +58,10 @@ private:
 
 	void ShowNewAnimSetDialog();
 
-	void OnFinishedChangingProperties(const FPropertyChangedEvent& PropertyChangedEvent);
-
 	TSharedPtr<SContextualAnimViewport> ViewportWidget;
 
 	TSharedPtr<SContextualAnimAssetBrowser> AssetBrowserWidget;
 
-	TSharedPtr<IDetailsView> EditingAssetWidget;
-
 	TSharedPtr<FContextualAnimPreviewScene> PreviewScene;
 
 	TSharedPtr<FContextualAnimViewModel> ViewModel;
diff --git a/Engine/Plugins/Experimental/Animation/ContextualAnimation/Source/ContextualAnimationEditor/Public/ContextualAnimViewModel.h b/Engine/Plugins/Experimental/Animation/ContextualAnimation/Source/ContextualAnimationEditor/Public/ContextualAnimViewModel.h
--- a/Engine/Plugins/Experimental/Animation/ContextualAnimation/Source/ContextualAnimationEditor/Public/ContextualAnimViewModel.h
+++ b/Engine/Plugins/Experimental/Animation/ContextualAnimation/Source/ContextualAnimationEditor/Public/ContextualAnimViewModel.h
@@ -64,6 +64,7 @@ public:
 	void OnFinishedChangingProperties(const FPropertyChangedEvent& PropertyChangedEvent);
 
 	void Initialize(UContextualAnimSceneAsset* InSceneAsset, const TSharedRef<FContextualAnimPreviewScene>& InPreviewScene);
+	void Shutdown();
 
 	void SetDefaultMode();
 	void SetNotifiesMode(const FContextualAnimTrack& AnimTrack);
@@ -79,16 +80,20 @@ public:
 	UContextualAnimSceneAsset* GetSceneAsset() const { return SceneAsset; }
 	UContextualAnimSceneInstance* GetSceneInstance() const { return SceneInstance.Get(); }
 	UContextualAnimMovieSceneTrack* FindTrackByRole(const FName& Role) const;
+	TSharedPtr<IDetailsView> GetDetailsView() { return DetailsView; }
 
 	void AddNewAnimSet(const FContextualAnimNewAnimSetParams& Params);
+	void RemoveSection(int32 SectionIdx);
+	void RemoveAnimSet(int32 SectionIdx, int32 AnimSetIdx);
 
 	void AddNewIKTarget(const UContextualAnimNewIKTargetParams& Params);
 
 	void AnimationModified(UAnimSequenceBase& Animation);
 
 	void SetActiveSection(int32 SectionIdx);
-
+	int32 GetActiveSection() const;
 	void SetActiveAnimSetForSection(int32 SectionIdx, int32 AnimSetIdx);
+	int32 GetActiveAnimSetForSection(int32 SectionIdx) const;
 
 	void ToggleSimulateMode();
 	bool IsSimulateModeInactive()	const { return SimulateModeState == ESimulateModeState::Inactive; }
@@ -119,6 +124,8 @@ public:
 	void DiscardChangeToActorTransformInScene();
 
 	float GetPlaybackTime() const;
+	
+	bool CanMakeEdits() const;
 
 private:
 
@@ -201,6 +208,8 @@ private:
 	/** Selected actor when the user starts modifying its transform in the scene */
 	TWeakObjectPtr<AActor> ModifyingTransformInSceneCachedActor;
 
+	TSharedPtr<IDetailsView> DetailsView;
+
 	AActor* SpawnPreviewActor(const FContextualAnimTrack& AnimTrack);
 
 	UWorld* GetWorld() const;
@@ -219,5 +228,12 @@ private:
 
 	void CreateSequencer();
 
+	void CreateDetailsView();
+
+	void RefreshDetailsView();
+
 	bool WantsToModifyMeshToSceneForSelectedActor() const;
+
+private:
+	bool bInitialized = false;
 };
\ No newline at end of file
diff --git a/Engine/Plugins/Experimental/Animation/ContextualAnimation/Source/ContextualAnimationEditor/Public/DetailCustomizations/ContextualAnimSceneAssetDetailCustom.h b/Engine/Plugins/Experimental/Animation/ContextualAnimation/Source/ContextualAnimationEditor/Public/DetailCustomizations/ContextualAnimSceneAssetDetailCustom.h
new file mode 100644
--- /dev/null
+++ b/Engine/Plugins/Experimental/Animation/ContextualAnimation/Source/ContextualAnimationEditor/Public/DetailCustomizations/ContextualAnimSceneAssetDetailCustom.h
@@ -0,0 +1,27 @@
+// Copyright Epic Games, Inc. All Rights Reserved.
+
+#pragma once
+
+#include "CoreMinimal.h"
+#include "UObject/WeakObjectPtr.h"
+#include "IDetailCustomization.h"
+
+class IDetailLayoutBuilder;
+class IDetailsView;
+class UContextualAnimSceneAsset;
+class FContextualAnimViewModel;
+
+class FContextualAnimSceneAssetDetailCustom : public IDetailCustomization
+{
+public:
+
+	static TSharedRef<IDetailCustomization> MakeInstance(TSharedRef<FContextualAnimViewModel> ViewModelRef);
+
+	FContextualAnimSceneAssetDetailCustom(TSharedRef<FContextualAnimViewModel> ViewModelRef);
+
+	virtual void CustomizeDetails(IDetailLayoutBuilder& DetailBuilder) override;
+
+protected:
+
+	TWeakPtr<FContextualAnimViewModel> ViewModelPtr;
+};
\ No newline at end of file
