Contextual Anim: Interactions are now locally predicted when started by an AutonomousProxy.

diff --git a/Engine/Plugins/Experimental/Animation/ContextualAnimation/Source/ContextualAnimation/Private/AnimNotifyState_EarlyOutContextualAnimWindow.cpp b/Engine/Plugins/Experimental/Animation/ContextualAnimation/Source/ContextualAnimation/Private/AnimNotifyState_EarlyOutContextualAnimWindow.cpp
new file mode 100644
--- /dev/null
+++ b/Engine/Plugins/Experimental/Animation/ContextualAnimation/Source/ContextualAnimation/Private/AnimNotifyState_EarlyOutContextualAnimWindow.cpp
@@ -0,0 +1,43 @@
+// Copyright Epic Games, Inc. All Rights Reserved.
+
+#include "AnimNotifyState_EarlyOutContextualAnimWindow.h"
+#include "Components/SkeletalMeshComponent.h"
+#include "GameFramework/Character.h"
+#include "GameFramework/CharacterMovementComponent.h"
+#include "ContextualAnimSceneActorComponent.h"
+
+#include UE_INLINE_GENERATED_CPP_BY_NAME(AnimNotifyState_EarlyOutContextualAnimWindow)
+
+UAnimNotifyState_EarlyOutContextualAnimWindow::UAnimNotifyState_EarlyOutContextualAnimWindow(const FObjectInitializer& ObjectInitializer)
+	: Super(ObjectInitializer)
+{
+
+}
+
+void UAnimNotifyState_EarlyOutContextualAnimWindow::NotifyTick(USkeletalMeshComponent* MeshComp, UAnimSequenceBase* Animation, float FrameDeltaTime, const FAnimNotifyEventReference& EventReference)
+{
+	if(MeshComp)
+	{
+		ACharacter* CharacterOwner = Cast<ACharacter>(MeshComp->GetOwner());
+		if(CharacterOwner && CharacterOwner->IsLocallyControlled())
+		{
+			if(UCharacterMovementComponent* MoveComp = CharacterOwner->GetCharacterMovement())
+			{
+				if(!MoveComp->GetLastInputVector().IsNearlyZero())
+				{
+					UContextualAnimSceneActorComponent* Comp = CharacterOwner->FindComponentByClass<UContextualAnimSceneActorComponent>();
+					if (Comp)
+					{
+						Comp->EarlyOutContextualAnimScene();
+					}
+				}
+			}
+		}
+		
+	}
+}
+
+FString UAnimNotifyState_EarlyOutContextualAnimWindow::GetNotifyName_Implementation() const
+{
+	return FString(TEXT("Early Out Contextual Anim"));
+}
diff --git a/Engine/Plugins/Experimental/Animation/ContextualAnimation/Source/ContextualAnimation/Private/ContextualAnimSceneActorComponent.cpp b/Engine/Plugins/Experimental/Animation/ContextualAnimation/Source/ContextualAnimation/Private/ContextualAnimSceneActorComponent.cpp
--- a/Engine/Plugins/Experimental/Animation/ContextualAnimation/Source/ContextualAnimation/Private/ContextualAnimSceneActorComponent.cpp
+++ b/Engine/Plugins/Experimental/Animation/ContextualAnimation/Source/ContextualAnimation/Private/ContextualAnimSceneActorComponent.cpp
@@ -23,7 +23,6 @@
 #if !(UE_BUILD_SHIPPING || UE_BUILD_TEST)
 TAutoConsoleVariable<int32> CVarContextualAnimIKDebug(TEXT("a.ContextualAnim.IK.Debug"), 0, TEXT("Draw Debug IK Targets"));
 TAutoConsoleVariable<float> CVarContextualAnimIKDrawDebugLifetime(TEXT("a.ContextualAnim.IK.DrawDebugLifetime"), 0, TEXT("Draw Debug Duration"));
-TAutoConsoleVariable<float> CVarContextualAnimIKForceAlpha(TEXT("a.ContextualAnim.IK.ForceAlpha"), -1.f, TEXT("Override Alpha value for all the targets. -1 = Disable"));
 #endif
 
 UContextualAnimSceneActorComponent::UContextualAnimSceneActorComponent(const FObjectInitializer& ObjectInitializer)
@@ -40,47 +39,124 @@ void UContextualAnimSceneActorComponent::GetLifetimeReplicatedProps(TArray< FLif
 
 	FDoRepLifetimeParams Params;
 	Params.bIsPushBased = true;
-	Params.Condition = COND_SimulatedOnly;
 	DOREPLIFETIME_WITH_PARAMS_FAST(UContextualAnimSceneActorComponent, RepBindings, Params);
 }
 
+bool UContextualAnimSceneActorComponent::IsOwnerLocallyControlled() const
+{
+	if (const APawn* OwnerPawn = Cast<APawn>(GetOwner()))
+	{
+		return OwnerPawn->IsLocallyControlled();
+	}
+
+	return false;
+}
+
 bool UContextualAnimSceneActorComponent::StartContextualAnimScene(const FContextualAnimSceneBindings& InBindings)
 {
-	UE_LOG(LogContextualAnim, Log, TEXT("%-21s UContextualAnimSceneActorComponent::StartContextualAnim Actor: %s"), *UEnum::GetValueAsString(TEXT("Engine.ENetRole"), GetOwner()->GetLocalRole()), *GetNameSafe(GetOwner()));
+	UE_LOG(LogContextualAnim, Log, TEXT("%-21s UContextualAnimSceneActorComponent::StartContextualAnim Actor: %s"),
+		*UEnum::GetValueAsString(TEXT("Engine.ENetRole"), GetOwner()->GetLocalRole()), *GetNameSafe(GetOwner()));
 
 	const FContextualAnimSceneBinding* OwnerBinding = InBindings.FindBindingByActor(GetOwner());
 	if (ensureAlways(OwnerBinding))
 	{
-		JoinScene(InBindings);
-
-		for (const FContextualAnimSceneBinding& Binding : InBindings)
+		if (GetOwner()->HasAuthority())
 		{
-			if (Binding.GetActor() != GetOwner())
+			JoinScene(InBindings);
+
+			for (const FContextualAnimSceneBinding& Binding : InBindings)
 			{
-				UContextualAnimSceneActorComponent* Comp = Binding.GetSceneActorComponent();
-				checkf(Comp, TEXT("Missing SceneActorComp on %s"), *GetNameSafe(Binding.GetActor()));
+				if (Binding.GetActor() != GetOwner())
+				{
+					UContextualAnimSceneActorComponent* Comp = Binding.GetSceneActorComponent();
+					checkf(Comp, TEXT("Missing SceneActorComp on %s"), *GetNameSafe(Binding.GetActor()));
 
-				Comp->JoinScene(InBindings);
+					Comp->JoinScene(InBindings);
+				}
 			}
-		}
 
-		//@TODO: Temp until we move the scene pivots to the bindings
-		UContextualAnimUtilities::BP_SceneBindings_AddOrUpdateWarpTargetsForBindings(InBindings);
+			//@TODO: Temp until we move the scene pivots to the bindings
+			UContextualAnimUtilities::BP_SceneBindings_AddOrUpdateWarpTargetsForBindings(InBindings);
 
-		if (GetOwner()->HasAuthority())
-		{
 			RepBindings = InBindings;
 			MARK_PROPERTY_DIRTY_FROM_NAME(UContextualAnimSceneActorComponent, RepBindings, this);
-
 			GetOwner()->ForceNetUpdate();
+
+			return true;
 		}
+		else if (GetOwner()->GetLocalRole() == ROLE_AutonomousProxy)
+		{
+			JoinScene(InBindings);
+
+			//@TODO: Temp until we move the scene pivots to the bindings
+			UContextualAnimUtilities::BP_SceneBindings_AddOrUpdateWarpTargetsForBindings(InBindings);
 
-		return true;
+			ServerStartContextualAnimScene(InBindings);
+
+			return true;
+		}
 	}
 
 	return false;
 }
 
+void UContextualAnimSceneActorComponent::ServerStartContextualAnimScene_Implementation(const FContextualAnimSceneBindings& InBindings)
+{
+	StartContextualAnimScene(InBindings);
+}
+
+bool UContextualAnimSceneActorComponent::ServerStartContextualAnimScene_Validate(const FContextualAnimSceneBindings& InBindings)
+{
+	return true;
+}
+
+void UContextualAnimSceneActorComponent::EarlyOutContextualAnimScene()
+{
+	if (const FContextualAnimSceneBinding* Binding = Bindings.FindBindingByActor(GetOwner()))
+	{
+		const UAnimInstance* AnimInstance = Binding->GetAnimInstance();
+		const UAnimMontage* ActiveMontage = AnimInstance ? AnimInstance->GetCurrentActiveMontage() : nullptr;
+		if (ActiveMontage)
+		{
+			UE_LOG(LogContextualAnim, Verbose, TEXT("%-21s UContextualAnimSceneActorComponent::EarlyOutContextualAnimScene Actor: %s ActiveMontage: %s"),
+				*UEnum::GetValueAsString(TEXT("Engine.ENetRole"), GetOwner()->GetLocalRole()), *GetNameSafe(GetOwner()), *GetNameSafe(ActiveMontage));
+
+			if (Bindings.GetAnimTrackFromBinding(*Binding).Animation == ActiveMontage)
+			{
+				// Stop animation.
+				LeaveScene();
+
+				// If we are on the server, rep bindings to stop animation on simulated proxies
+				if (GetOwner()->HasAuthority())
+				{
+					if (RepBindings.IsValid())
+					{
+						RepBindings.Clear();
+						MARK_PROPERTY_DIRTY_FROM_NAME(UContextualAnimSceneActorComponent, RepBindings, this);
+
+						GetOwner()->ForceNetUpdate();
+					}
+				}
+				// If local player, tell the server to stop the animation too
+				else if (GetOwner()->GetLocalRole() == ROLE_AutonomousProxy)
+				{
+					ServerEarlyOutContextualAnimScene();
+				}
+			}
+		}
+	}
+}
+
+void UContextualAnimSceneActorComponent::ServerEarlyOutContextualAnimScene_Implementation()
+{
+	EarlyOutContextualAnimScene();
+}
+
+bool UContextualAnimSceneActorComponent::ServerEarlyOutContextualAnimScene_Validate()
+{
+	return true;
+}
+
 void UContextualAnimSceneActorComponent::OnRep_Bindings(const FContextualAnimSceneBindings& LastRepBindings)
 {
 	// @TODO: This need more investigation but for now it prevents an issue caused by this OnRep_ triggering even when there is no (obvious) change in the data
@@ -102,9 +178,12 @@ void UContextualAnimSceneActorComponent::OnRep_Bindings(const FContextualAnimSce
 		if (ensureAlways(OwnerBinding))
 		{
 			// Join the scene (start playing animation, etc.)
-			JoinScene(RepBindings);
+			if (GetOwner()->GetLocalRole() != ROLE_AutonomousProxy)
+			{
+				JoinScene(RepBindings);
+			}
 
-			// RepBindings is only replicated from the initiator of the actor. 
+			// RepBindings is only replicated from the initiator of the action.
 			// So now we have to tell everyone else involved in the interaction to join us
 			// @TODO: For now this assumes that all the actors will start playing the animation at the same time. 
 			// We will expand this in the future to allow 'late' join
@@ -126,7 +205,7 @@ void UContextualAnimSceneActorComponent::OnRep_Bindings(const FContextualAnimSce
 		// In this case we don't want to tell everyone else to also leave the scene since there is very common for the initiator, 
 		// specially if is player character, to end the animation earlier for responsiveness
 		// It is more likely this will do nothing since we listen to montage end also on Simulated Proxies to 'predict' the end of the interaction.
-		if(RepBindings.GetID() == Bindings.GetID())
+		if (RepBindings.GetID() == Bindings.GetID() && GetOwner()->GetLocalRole() != ROLE_AutonomousProxy)
 		{
 			LeaveScene();
 		}
@@ -207,9 +286,9 @@ void UContextualAnimSceneActorComponent::OnJoinedScene(const FContextualAnimScen
 		// Prevent physics rotation. During the interaction we want to be fully root motion driven
 		if (UCharacterMovementComponent* MovementComp = GetOwner()->FindComponentByClass<UCharacterMovementComponent>())
 		{
-			RotationPropertiesBackup.bAllowPhysicsRotationDuringAnimRootMotion = MovementComp->bAllowPhysicsRotationDuringAnimRootMotion;
-			RotationPropertiesBackup.bUseControllerDesiredRotation = MovementComp->bUseControllerDesiredRotation;
-			RotationPropertiesBackup.bOrientRotationToMovement = MovementComp->bOrientRotationToMovement;
+			CharacterPropertiesBackup.bAllowPhysicsRotationDuringAnimRootMotion = MovementComp->bAllowPhysicsRotationDuringAnimRootMotion;
+			CharacterPropertiesBackup.bUseControllerDesiredRotation = MovementComp->bUseControllerDesiredRotation;
+			CharacterPropertiesBackup.bOrientRotationToMovement = MovementComp->bOrientRotationToMovement;
 			MovementComp->bAllowPhysicsRotationDuringAnimRootMotion = false;
 			MovementComp->bUseControllerDesiredRotation = false;
 			MovementComp->bOrientRotationToMovement = false;
@@ -241,9 +320,9 @@ void UContextualAnimSceneActorComponent::OnLeftScene()
 		// Restore bAllowPhysicsRotationDuringAnimRootMotion
 		if (UCharacterMovementComponent* MovementComp = GetOwner()->FindComponentByClass<UCharacterMovementComponent>())
 		{
-			MovementComp->bAllowPhysicsRotationDuringAnimRootMotion = RotationPropertiesBackup.bAllowPhysicsRotationDuringAnimRootMotion;
-			MovementComp->bUseControllerDesiredRotation = RotationPropertiesBackup.bUseControllerDesiredRotation;
-			MovementComp->bOrientRotationToMovement = RotationPropertiesBackup.bOrientRotationToMovement;
+			MovementComp->bAllowPhysicsRotationDuringAnimRootMotion = CharacterPropertiesBackup.bAllowPhysicsRotationDuringAnimRootMotion;
+			MovementComp->bUseControllerDesiredRotation = CharacterPropertiesBackup.bUseControllerDesiredRotation;
+			MovementComp->bOrientRotationToMovement = CharacterPropertiesBackup.bOrientRotationToMovement;
 		}
 
 		OnLeftSceneDelegate.Broadcast(this);
@@ -286,12 +365,18 @@ void UContextualAnimSceneActorComponent::JoinScene(const FContextualAnimSceneBin
 		// Disable collision between actors so they can align perfectly
 		SetIgnoreCollisionWithOtherActors(true);
 
-		// Prevent physics rotation. During the interaction we want to be fully root motion driven
 		if (UCharacterMovementComponent* MovementComp = GetOwner()->FindComponentByClass<UCharacterMovementComponent>())
 		{
-			RotationPropertiesBackup.bAllowPhysicsRotationDuringAnimRootMotion = MovementComp->bAllowPhysicsRotationDuringAnimRootMotion;
-			RotationPropertiesBackup.bUseControllerDesiredRotation = MovementComp->bUseControllerDesiredRotation;
-			RotationPropertiesBackup.bOrientRotationToMovement = MovementComp->bOrientRotationToMovement;
+			// Save movement state before the interaction starts so we can restore it when it ends
+			CharacterPropertiesBackup.bIgnoreClientMovementErrorChecksAndCorrection = MovementComp->bIgnoreClientMovementErrorChecksAndCorrection;
+			CharacterPropertiesBackup.bAllowPhysicsRotationDuringAnimRootMotion = MovementComp->bAllowPhysicsRotationDuringAnimRootMotion;
+			CharacterPropertiesBackup.bUseControllerDesiredRotation = MovementComp->bUseControllerDesiredRotation;
+			CharacterPropertiesBackup.bOrientRotationToMovement = MovementComp->bOrientRotationToMovement;
+
+			// Disable movement correction.
+			MovementComp->bIgnoreClientMovementErrorChecksAndCorrection = true;
+
+			// Prevent physics rotation. During the interaction we want to be fully root motion driven
 			MovementComp->bAllowPhysicsRotationDuringAnimRootMotion = false;
 			MovementComp->bUseControllerDesiredRotation = false;
 			MovementComp->bOrientRotationToMovement = false;
@@ -342,12 +427,13 @@ void UContextualAnimSceneActorComponent::LeaveScene()
 		// We might want to add a more robust mechanism to avoid overriding a request to disable collision that may have been set by another system
 		SetIgnoreCollisionWithOtherActors(false);
 
-		// Restore bAllowPhysicsRotationDuringAnimRootMotion
 		if (UCharacterMovementComponent* MovementComp = GetOwner()->FindComponentByClass<UCharacterMovementComponent>())
 		{
-			MovementComp->bAllowPhysicsRotationDuringAnimRootMotion = RotationPropertiesBackup.bAllowPhysicsRotationDuringAnimRootMotion;
-			MovementComp->bUseControllerDesiredRotation = RotationPropertiesBackup.bUseControllerDesiredRotation;
-			MovementComp->bOrientRotationToMovement = RotationPropertiesBackup.bOrientRotationToMovement;
+			// Restore movement state
+			MovementComp->bIgnoreClientMovementErrorChecksAndCorrection = CharacterPropertiesBackup.bIgnoreClientMovementErrorChecksAndCorrection;
+			MovementComp->bAllowPhysicsRotationDuringAnimRootMotion = CharacterPropertiesBackup.bAllowPhysicsRotationDuringAnimRootMotion;
+			MovementComp->bUseControllerDesiredRotation = CharacterPropertiesBackup.bUseControllerDesiredRotation;
+			MovementComp->bOrientRotationToMovement = CharacterPropertiesBackup.bOrientRotationToMovement;
 
 			//@TODO: Temp solution that assumes these interactions are not locally predicted and that is ok to be in flying mode during the entire animation
 			if (AnimTrack.bRequireFlyingMode && MovementComp->MovementMode == MOVE_Flying)
@@ -383,8 +469,6 @@ void UContextualAnimSceneActorComponent::OnMontageBlendingOut(UAnimMontage* Mont
 
 					GetOwner()->ForceNetUpdate();
 				}
-
-				//@TODO: Replicate this event separately for each other member of the interaction
 			}
 		}
 	}
@@ -438,14 +522,6 @@ void UContextualAnimSceneActorComponent::UpdateIKTargets()
 	const TArray<FContextualAnimIKTargetDefinition>& IKTargetDefs = Bindings.GetIKTargetDefContainerFromBinding(*BindingPtr).IKTargetDefs;
 	for (const FContextualAnimIKTargetDefinition& IKTargetDef : IKTargetDefs)
 	{
-#if !(UE_BUILD_SHIPPING || UE_BUILD_TEST)
-		const bool bDrawDebugEnable = CVarContextualAnimIKDebug.GetValueOnGameThread() > 0;
-		const float DrawDebugDuration = CVarContextualAnimIKDrawDebugLifetime.GetValueOnGameThread();
-		FTransform IKTargetParentTransformForDebug = FTransform::Identity;
-#endif
-
-		FTransform IKTargetTransform = FTransform::Identity;
-
 		float Alpha = UAnimNotifyState_IKWindow::GetIKAlphaValue(IKTargetDef.GoalName, MontageInstance);
 
 		// @TODO: IKTargetTransform will be off by 1 frame if we tick before target. 
@@ -453,60 +529,60 @@ void UContextualAnimSceneActorComponent::UpdateIKTargets()
 
 		if (const FContextualAnimSceneBinding* TargetBinding = Bindings.FindBindingByRole(IKTargetDef.TargetRoleName))
 		{
-			if (const USkeletalMeshComponent* TargetSkelMeshComp = TargetBinding->GetSkeletalMeshComponent())
+			// Do not update if the target actor should be playing and animation but its not yet. 
+			// This could happen in multi player when the initiator start playing the animation locally
+			const UAnimSequenceBase* TargetAnimation = Bindings.GetAnimTrackFromBinding(*TargetBinding).Animation;
+			if (TargetAnimation)
+			{
+				//@TODO: Add support for dynamic montages
+				const FAnimMontageInstance* TargetMontageInstance = TargetBinding->GetAnimMontageInstance();
+				if (!TargetMontageInstance || TargetMontageInstance->Montage != TargetAnimation)
+				{
+					Alpha = 0.f;
+				}
+			}
+
+			if (Alpha > 0.f)
 			{
-				if (IKTargetDef.Provider == EContextualAnimIKTargetProvider::Autogenerated)
+				if (const USkeletalMeshComponent* TargetSkelMeshComp = TargetBinding->GetSkeletalMeshComponent())
 				{
-					const FTransform IKTargetParentTransform = TargetSkelMeshComp->GetSocketTransform(IKTargetDef.TargetBoneName);
+					if (IKTargetDef.Provider == EContextualAnimIKTargetProvider::Autogenerated)
+					{
+						const FTransform IKTargetParentTransform = TargetSkelMeshComp->GetSocketTransform(IKTargetDef.TargetBoneName);
 
-					const float Time = MontageInstance->GetPosition();
-					IKTargetTransform = Bindings.GetIKTargetTransformFromBinding(*BindingPtr, IKTargetDef.GoalName, Time) * IKTargetParentTransform;
+						const float Time = MontageInstance->GetPosition();
+						const FTransform IKTargetTransform = Bindings.GetIKTargetTransformFromBinding(*BindingPtr, IKTargetDef.GoalName, Time) * IKTargetParentTransform;
+
+						IKTargets.Add(FContextualAnimIKTarget(IKTargetDef.GoalName, Alpha, IKTargetTransform));
 
 #if !(UE_BUILD_SHIPPING || UE_BUILD_TEST)
-					if (bDrawDebugEnable)
-					{
-						IKTargetParentTransformForDebug = IKTargetParentTransform;
-					}
+						if (CVarContextualAnimIKDebug.GetValueOnGameThread() > 0)
+						{
+							const float DrawDebugDuration = CVarContextualAnimIKDrawDebugLifetime.GetValueOnGameThread();
+							DrawDebugLine(GetWorld(), IKTargetParentTransform.GetLocation(), IKTargetTransform.GetLocation(), FColor::MakeRedToGreenColorFromScalar(Alpha), false, DrawDebugDuration, 0, 0.5f);
+							DrawDebugCoordinateSystem(GetWorld(), IKTargetTransform.GetLocation(), IKTargetTransform.Rotator(), 10.f, false, DrawDebugDuration, 0, 0.5f);
+						}
 #endif
-				}
-				else if (IKTargetDef.Provider == EContextualAnimIKTargetProvider::Bone)
-				{
-					IKTargetTransform = TargetSkelMeshComp->GetSocketTransform(IKTargetDef.TargetBoneName);
+					}
+					else if (IKTargetDef.Provider == EContextualAnimIKTargetProvider::Bone)
+					{
+						const FTransform IKTargetTransform = TargetSkelMeshComp->GetSocketTransform(IKTargetDef.TargetBoneName);
+
+						IKTargets.Add(FContextualAnimIKTarget(IKTargetDef.GoalName, Alpha, IKTargetTransform));
 
 #if !(UE_BUILD_SHIPPING || UE_BUILD_TEST)
-					if (bDrawDebugEnable)
-					{
-						IKTargetParentTransformForDebug = TargetSkelMeshComp->GetSocketTransform(TargetSkelMeshComp->GetParentBone(IKTargetDef.TargetBoneName));
-					}
+						if (CVarContextualAnimIKDebug.GetValueOnGameThread() > 0)
+						{
+							const float DrawDebugDuration = CVarContextualAnimIKDrawDebugLifetime.GetValueOnGameThread();
+							const FTransform IKTargetParentTransform = TargetSkelMeshComp->GetSocketTransform(TargetSkelMeshComp->GetParentBone(IKTargetDef.TargetBoneName));
+							DrawDebugLine(GetWorld(), IKTargetParentTransform.GetLocation(), IKTargetTransform.GetLocation(), FColor::MakeRedToGreenColorFromScalar(Alpha), false, DrawDebugDuration, 0, 0.5f);
+							DrawDebugCoordinateSystem(GetWorld(), IKTargetTransform.GetLocation(), IKTargetTransform.Rotator(), 10.f, false, DrawDebugDuration, 0, 0.5f);
+						}
 #endif
+					}
 				}
 			}
 		}
-
-#if !(UE_BUILD_SHIPPING || UE_BUILD_TEST)
-		const float ForcedAlphaValue = CVarContextualAnimIKForceAlpha.GetValueOnGameThread();
-		if (ForcedAlphaValue > 0.f)
-		{
-			Alpha = FMath::Clamp(ForcedAlphaValue, 0.f, 1.f);
-		}
-
-		if (bDrawDebugEnable)
-		{
-			const float DrawThickness = 0.5f;
-			const FColor DrawColor = FColor::MakeRedToGreenColorFromScalar(Alpha);
-
-			DrawDebugLine(GetWorld(), IKTargetParentTransformForDebug.GetLocation(), IKTargetTransform.GetLocation(), DrawColor, false, DrawDebugDuration, 0, DrawThickness);
-			DrawDebugCoordinateSystem(GetWorld(), IKTargetTransform.GetLocation(), IKTargetTransform.Rotator(), 10.f, false, DrawDebugDuration, 0, DrawThickness);
-			//DrawDebugSphere(GetWorld(), IKTargetTransform.GetLocation(), 5.f, 12, DrawColor, false, 0.f, 0, DrawThickness  );
-
-			//DrawDebugString(GetWorld(), IKTargetTransform.GetLocation(), FString::Printf(TEXT("%s (%f)"), *IKTargetDef.AlphaCurveName.ToString(), Alpha));
-		}
-#endif
-
-		// Convert IK Target to mesh space
-		//IKTargetTransform = IKTargetTransform.GetRelativeTransform(BindingPtr->GetSkeletalMeshComponent()->GetComponentTransform());
-
-		IKTargets.Add(FContextualAnimIKTarget(IKTargetDef.GoalName, Alpha, IKTargetTransform));
 	}
 }
 
diff --git a/Engine/Plugins/Experimental/Animation/ContextualAnimation/Source/ContextualAnimation/Public/AnimNotifyState_EarlyOutContextualAnimWindow.h b/Engine/Plugins/Experimental/Animation/ContextualAnimation/Source/ContextualAnimation/Public/AnimNotifyState_EarlyOutContextualAnimWindow.h
new file mode 100644
--- /dev/null
+++ b/Engine/Plugins/Experimental/Animation/ContextualAnimation/Source/ContextualAnimation/Public/AnimNotifyState_EarlyOutContextualAnimWindow.h
@@ -0,0 +1,21 @@
+// Copyright Epic Games, Inc. All Rights Reserved.
+
+#pragma once
+
+#include "CoreMinimal.h"
+#include "Animation/AnimNotifies/AnimNotifyState.h"
+#include "AnimNotifyState_EarlyOutContextualAnimWindow.generated.h"
+
+/** Notify used to allow player to early out from a contextual anim. Usually used at the end of the animations to improve responsivess */
+UCLASS(meta = (DisplayName = "Early Out Contextual Anim"))
+class CONTEXTUALANIMATION_API UAnimNotifyState_EarlyOutContextualAnimWindow : public UAnimNotifyState
+{
+	GENERATED_BODY()
+
+public:
+
+	UAnimNotifyState_EarlyOutContextualAnimWindow(const FObjectInitializer& ObjectInitializer);
+
+	virtual void NotifyTick(USkeletalMeshComponent* MeshComp, UAnimSequenceBase* Animation, float FrameDeltaTime, const FAnimNotifyEventReference& EventReference) override;
+	virtual FString GetNotifyName_Implementation() const override;
+};
\ No newline at end of file
diff --git a/Engine/Plugins/Experimental/Animation/ContextualAnimation/Source/ContextualAnimation/Public/ContextualAnimSceneActorComponent.h b/Engine/Plugins/Experimental/Animation/ContextualAnimation/Source/ContextualAnimation/Public/ContextualAnimSceneActorComponent.h
--- a/Engine/Plugins/Experimental/Animation/ContextualAnimation/Source/ContextualAnimation/Public/ContextualAnimSceneActorComponent.h
+++ b/Engine/Plugins/Experimental/Animation/ContextualAnimation/Source/ContextualAnimation/Public/ContextualAnimSceneActorComponent.h
@@ -63,6 +63,10 @@ public:
 	UFUNCTION(BlueprintCallable, Category = "Contextual Anim|Scene Actor Component")
 	bool StartContextualAnimScene(const FContextualAnimSceneBindings& InBindings);
 
+	void EarlyOutContextualAnimScene();
+
+	bool IsOwnerLocallyControlled() const;
+
 protected:
 
 	/** 
@@ -83,13 +87,14 @@ protected:
 	UPROPERTY(Transient)
 	TArray<FContextualAnimIKTarget> IKTargets;
 
-	struct FRotationProperties
+	struct FCharacterProperties
 	{
+		bool bIgnoreClientMovementErrorChecksAndCorrection = false;
 		bool bAllowPhysicsRotationDuringAnimRootMotion = false;
 		bool bUseControllerDesiredRotation = false;
 		bool bOrientRotationToMovement = false;
 	};
-	FRotationProperties RotationPropertiesBackup;
+	FCharacterProperties CharacterPropertiesBackup;
 
 	void UpdateIKTargets();
 
@@ -115,6 +120,12 @@ protected:
 
 	void LeaveScene();
 
+	UFUNCTION(Server, Reliable, WithValidation)
+	void ServerStartContextualAnimScene(const FContextualAnimSceneBindings& InBindings);
+
+	UFUNCTION(Server, Reliable, WithValidation)
+	void ServerEarlyOutContextualAnimScene();
+
 private:
 
 	bool bRegistered = false;
