Contextual Anim Editor: Fixed crash when opening a new asset

diff --git a/Engine/Plugins/Experimental/Animation/ContextualAnimation/Source/ContextualAnimationEditor/Private/DetailCustomizations/ContextualAnimSceneAssetDetailCustom.cpp b/Engine/Plugins/Experimental/Animation/ContextualAnimation/Source/ContextualAnimationEditor/Private/DetailCustomizations/ContextualAnimSceneAssetDetailCustom.cpp
--- a/Engine/Plugins/Experimental/Animation/ContextualAnimation/Source/ContextualAnimationEditor/Private/DetailCustomizations/ContextualAnimSceneAssetDetailCustom.cpp
+++ b/Engine/Plugins/Experimental/Animation/ContextualAnimation/Source/ContextualAnimationEditor/Private/DetailCustomizations/ContextualAnimSceneAssetDetailCustom.cpp
@@ -39,7 +39,7 @@ void FContextualAnimSceneAssetDetailCustom::CustomizeDetails(IDetailLayoutBuilde
 	Category.SetCategoryVisibility(false);
 
 	const int32 SectionIdx = ViewModel->GetActiveSection();
-	if (SectionIdx != INDEX_NONE)
+	if (const FContextualAnimSceneSection* Section = SceneAsset->GetSection(SectionIdx))
 	{
 		// Create new category to show editable data only for the currently active section in the editor
 		IDetailCategoryBuilder& CurrentSectionCategory = DetailBuilder.EditCategory(TEXT("Active Section"));
@@ -76,7 +76,7 @@ void FContextualAnimSceneAssetDetailCustom::CustomizeDetails(IDetailLayoutBuilde
 
 		// Same as above but for the currently active AnimSet
 		const int32 AnimSetIdx = ViewModel->GetActiveAnimSetForSection(SectionIdx);
-		if (AnimSetIdx != INDEX_NONE)
+		if (const FContextualAnimSet* AnimSet = Section->GetAnimSet(AnimSetIdx))
 		{
 			// Create new category to show editable data only for the currently active AnimSet in the editor
 			IDetailCategoryBuilder& CurrentAnimSetCategory = DetailBuilder.EditCategory(TEXT("Active AnimSet"));
