Fixed issue with contextual anim scene actor component having query collisions enabled

diff --git a/Engine/Plugins/Experimental/Animation/ContextualAnimation/Source/ContextualAnimation/Private/ContextualAnimSceneActorComponent.cpp b/Engine/Plugins/Experimental/Animation/ContextualAnimation/Source/ContextualAnimation/Private/ContextualAnimSceneActorComponent.cpp
--- a/Engine/Plugins/Experimental/Animation/ContextualAnimation/Source/ContextualAnimation/Private/ContextualAnimSceneActorComponent.cpp
+++ b/Engine/Plugins/Experimental/Animation/ContextualAnimation/Source/ContextualAnimation/Private/ContextualAnimSceneActorComponent.cpp
@@ -31,6 +31,8 @@ UContextualAnimSceneActorComponent::UContextualAnimSceneActorComponent(const FOb
 	PrimaryComponentTick.bCanEverTick = false;
 	PrimaryComponentTick.bStartWithTickEnabled = false;
 	SetIsReplicatedByDefault(true);
+
+	SetCollisionEnabled(ECollisionEnabled::NoCollision);
 }
 
 void UContextualAnimSceneActorComponent::GetLifetimeReplicatedProps(TArray< FLifetimeProperty >& OutLifetimeProps) const
