Contextual Anim Editor: Set ShouldSimulatePhysics for the preview scene to false and removed unnecessary call to SetIsTemporarilyHiddenInEditor

diff --git a/Engine/Plugins/Experimental/Animation/ContextualAnimation/Source/ContextualAnimationEditor/Private/ContextualAnimAssetEditorToolkit.cpp b/Engine/Plugins/Experimental/Animation/ContextualAnimation/Source/ContextualAnimationEditor/Private/ContextualAnimAssetEditorToolkit.cpp
--- a/Engine/Plugins/Experimental/Animation/ContextualAnimation/Source/ContextualAnimationEditor/Private/ContextualAnimAssetEditorToolkit.cpp
+++ b/Engine/Plugins/Experimental/Animation/ContextualAnimation/Source/ContextualAnimationEditor/Private/ContextualAnimAssetEditorToolkit.cpp
@@ -71,11 +71,8 @@ void FContextualAnimAssetEditorToolkit::InitAssetEditor(const EToolkitMode::Type
 	// Create Preview Scene
 	if (!PreviewScene.IsValid())
 	{
-		PreviewScene = MakeShareable(new FContextualAnimPreviewScene(FPreviewScene::ConstructionValues().AllowAudioPlayback(true).ShouldSimulatePhysics(true).ForceUseMovementComponentInNonGameWorld(true),
+		PreviewScene = MakeShareable(new FContextualAnimPreviewScene(FPreviewScene::ConstructionValues().AllowAudioPlayback(true).ShouldSimulatePhysics(false).ForceUseMovementComponentInNonGameWorld(true),
 			StaticCastSharedRef<FContextualAnimAssetEditorToolkit>(AsShared())));
-
-		//Temporary fix for missing attached assets - MDW (Copied from FPersonaToolkit::CreatePreviewScene)
-		PreviewScene->GetWorld()->GetWorldSettings()->SetIsTemporarilyHiddenInEditor(false);
 	}
 
 	// Create viewport widget
