Contextual Anim: Refactor selection logic to consider all valid sets and pick one randomly based on a weight and added support for optional actors.

diff --git a/Engine/Plugins/Experimental/Animation/ContextualAnimation/Source/ContextualAnimation/Private/ContextualAnimSceneAsset.cpp b/Engine/Plugins/Experimental/Animation/ContextualAnimation/Source/ContextualAnimation/Private/ContextualAnimSceneAsset.cpp
--- a/Engine/Plugins/Experimental/Animation/ContextualAnimation/Source/ContextualAnimation/Private/ContextualAnimSceneAsset.cpp
+++ b/Engine/Plugins/Experimental/Animation/ContextualAnimation/Source/ContextualAnimation/Private/ContextualAnimSceneAsset.cpp
@@ -53,6 +53,22 @@ static void ExtractPoseIgnoringForceRootLock(UAnimSequenceBase* AnimSequenceBase
 	}
 }
 
+// FContextualAnimSet
+//==============================================================================================
+
+int32 FContextualAnimSet::GetNumMandatoryRoles() const
+{
+	int32 Result = 0;
+	for (const FContextualAnimTrack& AnimTrack : Tracks)
+	{
+		if (!AnimTrack.bOptional)
+		{
+			Result++;
+		}
+	}
+	return Result;
+}
+
 // FContextualAnimSceneSection
 //==============================================================================================
 
@@ -541,6 +557,12 @@ const FContextualAnimSceneSection* UContextualAnimSceneAsset::GetSection(const F
 	return Sections.FindByPredicate([&SectionName](const FContextualAnimSceneSection& Section) { return Section.Name == SectionName; });
 }
 
+const FContextualAnimSet* UContextualAnimSceneAsset::GetAnimSet(int32 SectionIdx, int32 AnimSetIdx) const
+{
+	const FContextualAnimSceneSection* Section = GetSection(SectionIdx);
+	return Section ? Section->GetAnimSet(AnimSetIdx) : nullptr;
+}
+
 int32 UContextualAnimSceneAsset::GetSectionIndex(const FName& SectionName) const
 {
 	return Sections.IndexOfByPredicate([&SectionName](const FContextualAnimSceneSection& Section) { return Section.Name == SectionName; });
@@ -588,6 +610,12 @@ TArray<FName> UContextualAnimSceneAsset::GetRoles() const
  	return Result;
 }
 
+int32 UContextualAnimSceneAsset::GetNumMandatoryRoles(int32 SectionIdx, int32 AnimSetIdx) const
+{
+	const FContextualAnimSet* AnimSet = GetAnimSet(SectionIdx, AnimSetIdx);
+	return AnimSet ? AnimSet->GetNumMandatoryRoles() : 0;
+}
+
 const FTransform& UContextualAnimSceneAsset::GetMeshToComponentForRole(const FName& Role) const
 {
 	if (RolesAsset)
diff --git a/Engine/Plugins/Experimental/Animation/ContextualAnimation/Source/ContextualAnimation/Private/ContextualAnimTypes.cpp b/Engine/Plugins/Experimental/Animation/ContextualAnimation/Source/ContextualAnimation/Private/ContextualAnimTypes.cpp
--- a/Engine/Plugins/Experimental/Animation/ContextualAnimation/Source/ContextualAnimation/Private/ContextualAnimTypes.cpp
+++ b/Engine/Plugins/Experimental/Animation/ContextualAnimation/Source/ContextualAnimation/Private/ContextualAnimTypes.cpp
@@ -476,203 +476,184 @@ bool FContextualAnimSceneBindings::NetSerialize(FArchive& Ar, class UPackageMap*
 	return true;
 }
 
-bool FContextualAnimSceneBindings::TryCreateBindings(const UContextualAnimSceneAsset& SceneAsset, int32 SectionIdx, int32 AnimSetIdx, const TMap<FName, FContextualAnimSceneBindingContext>& Params, FContextualAnimSceneBindings& OutBindings)
+bool FContextualAnimSceneBindings::CheckConditions(const UContextualAnimSceneAsset& SceneAsset, int32 SectionIdx, int32 AnimSetIdx, const TMap<FName, FContextualAnimSceneBindingContext>& Params)
 {
 	check(SceneAsset.HasValidData());
 
-	OutBindings = FContextualAnimSceneBindings(SceneAsset, SectionIdx, AnimSetIdx);
-
-	// Find the actor that should be bound to the primary Role.
-	const FName PrimaryRole = SceneAsset.GetPrimaryRole();
-	const FContextualAnimSceneBindingContext* PrimaryPtr = Params.Find(PrimaryRole);
-	if (PrimaryPtr == nullptr)
+	auto DoCheck = [&SceneAsset, SectionIdx, AnimSetIdx](const FContextualAnimSceneBindingContext& Primary, const FContextualAnimSceneBindingContext& Querier, const FName& Role)
 	{
-		UE_LOG(LogContextualAnim, Warning, TEXT("FContextualAnimSceneBindings::TryCreateBindings Failed. Reason: Can't find valid actor for primary role. SceneAsset: %s PrimaryRole: %s"),
-			*GetNameSafe(&SceneAsset), *PrimaryRole.ToString());
+		if (AActor* Actor = Querier.GetActor())
+		{
+			UContextualAnimSceneActorComponent* SceneActorComp = Actor->FindComponentByClass<UContextualAnimSceneActorComponent>();
+			if (SceneActorComp == nullptr)
+			{
+				UE_LOG(LogContextualAnim, Verbose, TEXT("FContextualAnimSceneBindings::CheckConditions Failed. Reason: Missing ContextualAnimSceneActorComp. SceneAsset: %s Actor: %s Role: %s SectionIdx: %d AnimSetIdx: %d"),
+					*GetNameSafe(&SceneAsset), *GetNameSafe(Actor), *Role.ToString(), SectionIdx, AnimSetIdx);
 
-		OutBindings.Reset();
-		return false;
-	}
+				return false;
+			}
+		}
 
-	if(AActor* PrimaryActor = PrimaryPtr->GetActor())
-	{
-		UContextualAnimSceneActorComponent* SceneActorComp = PrimaryActor->FindComponentByClass<UContextualAnimSceneActorComponent>();
-		if(SceneActorComp == nullptr)
+		const FContextualAnimTrack* AnimTrack = SceneAsset.GetAnimTrack(SectionIdx, AnimSetIdx, Role);
+		if (AnimTrack == nullptr || AnimTrack->DoesQuerierPassSelectionCriteria(Primary, Querier) == false)
 		{
-			UE_LOG(LogContextualAnim, Warning, TEXT("FContextualAnimSceneBindings::TryCreateBindings Failed. Reason: Missing ContextualAnimSceneActorComp on primary actor. SceneAsset: %s Actor: %s"),
-				*GetNameSafe(&SceneAsset), *GetNameSafe(PrimaryActor));
+			UE_LOG(LogContextualAnim, Verbose, TEXT("FContextualAnimSceneBindings::CheckConditions Failed. Reason: Reason: Can't find valid track for actor. SceneAsset: %s Actor: %s Role: %s SectionIdx: %d AnimSetIdx: %d"),
+				*GetNameSafe(&SceneAsset), *GetNameSafe(Querier.GetActor()), *Role.ToString(), SectionIdx, AnimSetIdx);
 
-			OutBindings.Reset();
 			return false;
 		}
-	}
 
-	// First, try to bind primary track. 
-	// @TODO: Revisit this, passing the same data twice (as primary and querier) feels weird, but this allow us to run the selection mechanism even on the primary actor.
+		return true;
+	};
 
-	const FContextualAnimTrack* PrimaryAnimTrack = SceneAsset.GetAnimTrack(SectionIdx, AnimSetIdx, PrimaryRole);
-	if (PrimaryAnimTrack && PrimaryAnimTrack->DoesQuerierPassSelectionCriteria(*PrimaryPtr, *PrimaryPtr))
-	{
-		OutBindings.Add(FContextualAnimSceneBinding(*PrimaryPtr, *PrimaryAnimTrack));
-	}
-	else
+	// Find the actor that should be bound to the primary Role.
+	const FName PrimaryRole = SceneAsset.GetPrimaryRole();
+	const FContextualAnimSceneBindingContext* PrimaryPtr = Params.Find(PrimaryRole);
+	if (PrimaryPtr == nullptr)
 	{
-		UE_LOG(LogContextualAnim, Warning, TEXT("FContextualAnimSceneBindings::TryCreateBindings Failed. Reason: Can't find valid track for primary actor. SceneAsset: %s Role: %s Actor: %s SectionIdx: %d AnimSetIdx: %d"),
-			*GetNameSafe(&SceneAsset), *PrimaryRole.ToString(), *GetNameSafe(PrimaryPtr->GetActor()), SectionIdx, AnimSetIdx);
+		UE_LOG(LogContextualAnim, Verbose, TEXT("FContextualAnimSceneBindings::CheckConditions Failed. Reason: Can't find valid actor for primary role. SceneAsset: %s PrimaryRole: %s"),
+			*GetNameSafe(&SceneAsset), *PrimaryRole.ToString());
 
-		OutBindings.Reset();
 		return false;
 	}
 
-	// Now try to bind secondary tracks
+	// Test primary actor first
+	// Passing the same data twice (as primary and querier) feels weird, but this allow us to run the selection mechanism even on the primary actor.
+	if (DoCheck(*PrimaryPtr, *PrimaryPtr, PrimaryRole) == false)
+	{
+		return false;
+	}
 
+	// Now test secondary actors
+	int32 ValidEntries = 1;
 	for (const auto& Pair : Params)
 	{
 		FName RoleToBind = Pair.Key;
 		if (RoleToBind != PrimaryRole)
 		{
-			if (AActor* SecondaryActor = Pair.Value.GetActor())
+			if (DoCheck(*PrimaryPtr, Pair.Value, RoleToBind) == false)
 			{
-				UContextualAnimSceneActorComponent* SceneActorComp = SecondaryActor->FindComponentByClass<UContextualAnimSceneActorComponent>();
-				if (SceneActorComp == nullptr)
-				{
-					UE_LOG(LogContextualAnim, Warning, TEXT("FContextualAnimSceneBindings::TryCreateBindings Failed. Reason: Missing ContextualAnimSceneActorComp on secondary actor. SceneAsset: %s Actor: %s"),
-						*GetNameSafe(&SceneAsset), *GetNameSafe(SecondaryActor));
-
-					OutBindings.Reset();
-					return false;
-				}
-			}
-
-			const FContextualAnimTrack* AnimTrack = SceneAsset.GetAnimTrack(SectionIdx, AnimSetIdx, RoleToBind);
-			if (AnimTrack && AnimTrack->DoesQuerierPassSelectionCriteria(*PrimaryPtr, Pair.Value))
-			{
-				OutBindings.Add(FContextualAnimSceneBinding(Pair.Value, *AnimTrack));
-			}
-			else
-			{
-				UE_LOG(LogContextualAnim, Warning, TEXT("FContextualAnimSceneBindings::TryCreateBindings Failed. Reason: Can't find valid track for secondary actor. SceneAsset: %s Role: %s Actor: %s SectionIdx: %d AnimSetIdx: %d"),
-					*GetNameSafe(&SceneAsset), *RoleToBind.ToString(), *GetNameSafe(Pair.Value.GetActor()), SectionIdx, AnimSetIdx);
-
-				OutBindings.Reset();
 				return false;
 			}
+
+			ValidEntries++;
 		}
 	}
 
-	if (OutBindings.Num() != SceneAsset.GetNumRoles())
+	const int32 NumMandatoryRoles = SceneAsset.GetNumMandatoryRoles(SectionIdx, AnimSetIdx);
+	if (ValidEntries < NumMandatoryRoles)
 	{
-		UE_LOG(LogContextualAnim, Warning, TEXT("FContextualAnimSceneBindings::TryCreateBindings Failed. Reason: Not all the roles were filled (%d/%d). SceneAsset: %s SectionIdx: %d AnimSetIdx: %d"),
-			OutBindings.Num(), SceneAsset.GetNumRoles(), *GetNameSafe(&SceneAsset), SectionIdx, AnimSetIdx);
-
-		OutBindings.Reset();
 		return false;
 	}
 
 	return true;
 }
 
+bool FContextualAnimSceneBindings::TryCreateBindings(const UContextualAnimSceneAsset& SceneAsset, int32 SectionIdx, int32 AnimSetIdx, const TMap<FName, FContextualAnimSceneBindingContext>& Params, FContextualAnimSceneBindings& OutBindings)
+{
+	check(SceneAsset.HasValidData());
+	
+	OutBindings.Reset();
+
+	if (FContextualAnimSceneBindings::CheckConditions(SceneAsset, SectionIdx, AnimSetIdx, Params))
+	{
+		OutBindings = FContextualAnimSceneBindings(SceneAsset, SectionIdx, AnimSetIdx);
+		for (const auto& Pair : Params)
+		{
+			const FContextualAnimTrack* AnimTrackPtr = SceneAsset.GetAnimTrack(SectionIdx, AnimSetIdx, Pair.Key);
+			check(AnimTrackPtr);
+
+			OutBindings.Data.Add(FContextualAnimSceneBinding(Pair.Value, *AnimTrackPtr));
+		}
+
+		check(OutBindings.IsValid());
+		return true;
+	}
+
+	return false;
+}
+
 bool FContextualAnimSceneBindings::TryCreateBindings(const UContextualAnimSceneAsset& SceneAsset, int32 SectionIdx, int32 AnimSetIdx, const FContextualAnimSceneBindingContext& Primary, const FContextualAnimSceneBindingContext& Secondary, FContextualAnimSceneBindings& OutBindings)
 {
 	check(SceneAsset.HasValidData());
 
-	OutBindings = FContextualAnimSceneBindings(SceneAsset, SectionIdx, AnimSetIdx);
+	OutBindings.Reset();
 
 	const TArray<FContextualAnimRoleDefinition>& Roles = SceneAsset.GetRolesAsset()->Roles;
-
 	if (Roles.Num() > 2)
 	{
 		UE_LOG(LogContextualAnim, Warning, TEXT("FContextualAnimSceneBindings::TryCreateBindings Failed. Reason: Trying to create bindings with two actors for a SceneAsset with more than two roles. SceneAsset: %s Num Roles: %d SectionIdx: %d AnimSetIdx: %d"),
 			*GetNameSafe(&SceneAsset), Roles.Num(), SectionIdx, AnimSetIdx);
 
-		OutBindings.Reset();
 		return false;
 	}
 
+	TMap<FName, FContextualAnimSceneBindingContext> Params;
 	for(const FContextualAnimRoleDefinition& RoleDef : Roles)
 	{
 		const FName PrimaryRole = SceneAsset.GetPrimaryRole();
-		if(RoleDef.Name == PrimaryRole)
-		{
-			if (AActor* Actor = Primary.GetActor())
-			{
-				UContextualAnimSceneActorComponent* SceneActorComp = Actor->FindComponentByClass<UContextualAnimSceneActorComponent>();
-				if (SceneActorComp == nullptr)
-				{
-					UE_LOG(LogContextualAnim, Warning, TEXT("FContextualAnimSceneBindings::TryCreateBindings Failed. Reason: Missing ContextualAnimSceneActorComp on primary actor. SceneAsset: %s Actor: %s"),
-						*GetNameSafe(&SceneAsset), *GetNameSafe(Actor));
+		Params.Add(RoleDef.Name, (RoleDef.Name == PrimaryRole) ? Primary : Secondary);
+	}
 
-					OutBindings.Reset();
-					return false;
-				}
-			}
+	return TryCreateBindings(SceneAsset, SectionIdx, AnimSetIdx, Params, OutBindings);
+}
 
-			const FContextualAnimTrack* AnimTrack = SceneAsset.GetAnimTrack(SectionIdx, AnimSetIdx, PrimaryRole);
-			if (AnimTrack && AnimTrack->DoesQuerierPassSelectionCriteria(Primary, Primary))
-			{
-				OutBindings.Add(FContextualAnimSceneBinding(Primary, *AnimTrack));
-			}
-			else
-			{
-				UE_LOG(LogContextualAnim, Warning, TEXT("FContextualAnimSceneBindings::TryCreateBindings Failed. Reason: Can't find valid track for primary actor. SceneAsset: %s Role: %s Actor: %s SectionIdx: %d AnimSetIdx: %d"),
-					*GetNameSafe(&SceneAsset), *RoleDef.Name.ToString(), *GetNameSafe(Primary.GetActor()), SectionIdx, AnimSetIdx);
+bool FContextualAnimSceneBindings::TryCreateBindings(const UContextualAnimSceneAsset& SceneAsset, int32 SectionIdx, const TMap<FName, FContextualAnimSceneBindingContext>& Params, FContextualAnimSceneBindings& OutBindings)
+{
+	check(SceneAsset.HasValidData());
+	OutBindings.Reset();
 
-				OutBindings.Reset();
-				return false;
-			}
+	int32 AnimSetIdxSelected = INDEX_NONE;
+
+	const int32 NumSets = SceneAsset.GetNumAnimSetsInSection(SectionIdx);
+	if (NumSets == 1)
+	{
+		if (CheckConditions(SceneAsset, SectionIdx, 0, Params))
+		{
+			AnimSetIdxSelected = 0;
 		}
-		else // Secondary Role
+	}
+	else if (NumSets > 1)
+	{
+		const FContextualAnimSceneSection* Section = SceneAsset.GetSection(SectionIdx);
+		TArray<TTuple<int32, float>, TInlineAllocator<5>> ValidSets; // 0: AnimSetIdx, 1: AnimSetRandomWeight
+		
+		float TotalWeight = 0;
+		for (int32 AnimSetIdx = 0; AnimSetIdx < NumSets; AnimSetIdx++)
 		{
-			if (AActor* Actor = Secondary.GetActor())
+			if (CheckConditions(SceneAsset, SectionIdx, AnimSetIdx, Params))
 			{
-				UContextualAnimSceneActorComponent* SceneActorComp = Actor->FindComponentByClass<UContextualAnimSceneActorComponent>();
-				if (SceneActorComp == nullptr)
-				{
-					UE_LOG(LogContextualAnim, Warning, TEXT("FContextualAnimSceneBindings::TryCreateBindings Failed. Reason: Missing ContextualAnimSceneActorComp on primary actor. SceneAsset: %s Actor: %s"),
-						*GetNameSafe(&SceneAsset), *GetNameSafe(Actor));
-
-					OutBindings.Reset();
-					return false;
-				}
+				const float AnimSetRandomWeight = FMath::Max(Section->GetAnimSet(AnimSetIdx)->RandomWeight, 0);
+				ValidSets.Add(MakeTuple(AnimSetIdx, AnimSetRandomWeight));
+				TotalWeight += AnimSetRandomWeight;
 			}
+		}
 
-			const FContextualAnimTrack* AnimTrack = SceneAsset.GetAnimTrack(SectionIdx, AnimSetIdx, RoleDef.Name);
-			if (AnimTrack && AnimTrack->DoesQuerierPassSelectionCriteria(Primary, Secondary))
-			{
-				OutBindings.Add(FContextualAnimSceneBinding(Secondary, *AnimTrack));
-			}
-			else
+		float RandomValue = FMath::RandRange(0.f, TotalWeight);
+		for (int32 Idx = 0; Idx < ValidSets.Num(); Idx++)
+		{
+			RandomValue -= FMath::Max(ValidSets[Idx].Get<1>(), 0);
+			if (RandomValue <= 0)
 			{
-				UE_LOG(LogContextualAnim, Warning, TEXT("FContextualAnimSceneBindings::TryCreateBindings Failed. Reason: Can't find valid track for secondary actor. SceneAsset: %s Role: %s Actor: %s SectionIdx: %d AnimSetIdx: %d"),
-					*GetNameSafe(&SceneAsset), *RoleDef.Name.ToString(), *GetNameSafe(Secondary.GetActor()), SectionIdx, AnimSetIdx);
-
-				OutBindings.Reset();
-				return false;
+				AnimSetIdxSelected = ValidSets[Idx].Get<0>();
+				break;
 			}
 		}
 	}
 
-	if (OutBindings.Num() != SceneAsset.GetNumRoles())
-	{
-		UE_LOG(LogContextualAnim, Warning, TEXT("FContextualAnimSceneBindings::TryCreateBindings Failed. Reason: Not all the roles were filled (%d/%d). SceneAsset: %s SectionIdx: %d AnimSetIdx: %d"),
-			OutBindings.Num(), SceneAsset.GetNumRoles(), *GetNameSafe(&SceneAsset), SectionIdx, AnimSetIdx);
-
-		OutBindings.Reset();
-		return false;
-	}
-
-	return true;
-}
-
-bool FContextualAnimSceneBindings::TryCreateBindings(const UContextualAnimSceneAsset& SceneAsset, int32 SectionIdx, const TMap<FName, FContextualAnimSceneBindingContext>& Params, FContextualAnimSceneBindings& OutBindings)
-{
-	const int32 NumSets = SceneAsset.GetNumAnimSetsInSection(SectionIdx);
-	for (int32 AnimSetIdx = 0; AnimSetIdx < NumSets; AnimSetIdx++)
+	if (AnimSetIdxSelected != INDEX_NONE)
 	{
-		if (TryCreateBindings(SceneAsset, SectionIdx, AnimSetIdx, Params, OutBindings))
+		OutBindings = FContextualAnimSceneBindings(SceneAsset, SectionIdx, AnimSetIdxSelected);
+		for (const auto& Pair : Params)
 		{
-			return true;
+			const FContextualAnimTrack* AnimTrackPtr = SceneAsset.GetAnimTrack(SectionIdx, AnimSetIdxSelected, Pair.Key);
+			check(AnimTrackPtr);
+
+			OutBindings.Data.Add(FContextualAnimSceneBinding(Pair.Value, *AnimTrackPtr));
 		}
+
+		check(OutBindings.IsValid());
+		return true;
 	}
 
 	return false;
@@ -680,16 +661,27 @@ bool FContextualAnimSceneBindings::TryCreateBindings(const UContextualAnimSceneA
 
 bool FContextualAnimSceneBindings::TryCreateBindings(const UContextualAnimSceneAsset& SceneAsset, int32 SectionIdx, const FContextualAnimSceneBindingContext& Primary, const FContextualAnimSceneBindingContext& Secondary, FContextualAnimSceneBindings& OutBindings)
 {
-	const int32 NumSets = SceneAsset.GetNumAnimSetsInSection(SectionIdx);
-	for (int32 AnimSetIdx = 0; AnimSetIdx < NumSets; AnimSetIdx++)
+	check(SceneAsset.HasValidData());
+
+	OutBindings.Reset();
+
+	const TArray<FContextualAnimRoleDefinition>& Roles = SceneAsset.GetRolesAsset()->Roles;
+	if (Roles.Num() > 2)
 	{
-		if (TryCreateBindings(SceneAsset, SectionIdx, AnimSetIdx, Primary, Secondary, OutBindings))
-		{
-			return true;
-		}
+		UE_LOG(LogContextualAnim, Warning, TEXT("FContextualAnimSceneBindings::TryCreateBindings Failed. Reason: Trying to create bindings with two actors for a SceneAsset with more than two roles. SceneAsset: %s Num Roles: %d SectionIdx: %d"),
+			*GetNameSafe(&SceneAsset), Roles.Num(), SectionIdx);
+
+		return false;
 	}
 
-	return false;
+	TMap<FName, FContextualAnimSceneBindingContext> Params;
+	for (const FContextualAnimRoleDefinition& RoleDef : Roles)
+	{
+		const FName PrimaryRole = SceneAsset.GetPrimaryRole();
+		Params.Add(RoleDef.Name, (RoleDef.Name == PrimaryRole) ? Primary : Secondary);
+	}
+
+	return TryCreateBindings(SceneAsset, SectionIdx, Params, OutBindings);
 }
 
 void FContextualAnimSceneBindings::CalculateAnimSetPivots(TArray<FContextualAnimSetPivot>& OutScenePivots) const
diff --git a/Engine/Plugins/Experimental/Animation/ContextualAnimation/Source/ContextualAnimation/Public/ContextualAnimSceneAsset.h b/Engine/Plugins/Experimental/Animation/ContextualAnimation/Source/ContextualAnimation/Public/ContextualAnimSceneAsset.h
--- a/Engine/Plugins/Experimental/Animation/ContextualAnimation/Source/ContextualAnimation/Public/ContextualAnimSceneAsset.h
+++ b/Engine/Plugins/Experimental/Animation/ContextualAnimation/Source/ContextualAnimation/Public/ContextualAnimSceneAsset.h
@@ -52,6 +52,8 @@ struct CONTEXTUALANIMATION_API FContextualAnimSet
 	/** Used by the selection mechanism to 'break the tie' when multiple Sets can be selected */
 	UPROPERTY(EditAnywhere, Category = "Defaults", meta = (ClampMin = "0", UIMin = "0", ClampMax = "1", UIMax = "1"))
 	float RandomWeight = 1.f;
+
+	int32 GetNumMandatoryRoles() const;
 };
 
 /** Named container with one or more ContextualAnimSet */
@@ -213,6 +215,8 @@ public:
 
 	int32 GetNumRoles() const { return RolesAsset ? RolesAsset->GetNumRoles() : 0; }
 
+	int32 GetNumMandatoryRoles(int32 SectionIdx, int32 AnimSetIdx) const;
+
 	const FTransform& GetMeshToComponentForRole(const FName& Role) const;
 
 	TArray<FName> GetSectionNames() const;
@@ -225,6 +229,8 @@ public:
 
 	const FContextualAnimSceneSection* GetSection(const FName& SectionName) const;
 
+	const FContextualAnimSet* GetAnimSet(int32 SectionIdx, int32 AnimSetIdx) const;
+
 	int32 GetSectionIndex(const FName& SectionName) const;
 	
  	const FContextualAnimTrack* GetAnimTrack(int32 SectionIdx, int32 AnimSetIdx, const FName& Role) const;
diff --git a/Engine/Plugins/Experimental/Animation/ContextualAnimation/Source/ContextualAnimation/Public/ContextualAnimTypes.h b/Engine/Plugins/Experimental/Animation/ContextualAnimation/Source/ContextualAnimation/Public/ContextualAnimTypes.h
--- a/Engine/Plugins/Experimental/Animation/ContextualAnimation/Source/ContextualAnimation/Public/ContextualAnimTypes.h
+++ b/Engine/Plugins/Experimental/Animation/ContextualAnimation/Source/ContextualAnimation/Public/ContextualAnimTypes.h
@@ -75,6 +75,10 @@ struct CONTEXTUALANIMATION_API FContextualAnimTrack
 	UPROPERTY(EditAnywhere, BlueprintReadWrite, Category = "Defaults")
 	bool bRequireFlyingMode = false;
 
+	/** Whether the actor that should play this animation is optional */
+	UPROPERTY(EditAnywhere, BlueprintReadWrite, Category = "Defaults")
+	bool bOptional = false;
+
 	/** Container for alignment tracks */
 	UPROPERTY()
 	FContextualAnimAlignmentTrackContainer AlignmentData;
@@ -412,6 +416,7 @@ struct CONTEXTUALANIMATION_API FContextualAnimSceneBindings
 	FORCEINLINE TArray<FContextualAnimSceneBinding>::RangedForIteratorType      end() { return Data.end(); }
 	FORCEINLINE TArray<FContextualAnimSceneBinding>::RangedForConstIteratorType end() const { return Data.end(); }
 
+	static bool CheckConditions(const UContextualAnimSceneAsset& SceneAsset, int32 SectionIdx, int32 AnimSetIdx, const TMap<FName, FContextualAnimSceneBindingContext>& Params);
 	static bool TryCreateBindings(const UContextualAnimSceneAsset& SceneAsset, int32 SectionIdx, int32 AnimSetIdx, const TMap<FName, FContextualAnimSceneBindingContext>& Params, FContextualAnimSceneBindings& OutBindings);
 	static bool TryCreateBindings(const UContextualAnimSceneAsset& SceneAsset, int32 SectionIdx, const TMap<FName, FContextualAnimSceneBindingContext>& Params, FContextualAnimSceneBindings& OutBindings);
 	static bool TryCreateBindings(const UContextualAnimSceneAsset& SceneAsset, int32 SectionIdx, int32 AnimSetIdx, const FContextualAnimSceneBindingContext& Primary, const FContextualAnimSceneBindingContext& Secondary, FContextualAnimSceneBindings& OutBindings);
diff --git a/Engine/Plugins/Experimental/Animation/ContextualAnimation/Source/ContextualAnimationEditor/Private/ContextualAnimViewModel.cpp b/Engine/Plugins/Experimental/Animation/ContextualAnimation/Source/ContextualAnimationEditor/Private/ContextualAnimViewModel.cpp
--- a/Engine/Plugins/Experimental/Animation/ContextualAnimation/Source/ContextualAnimationEditor/Private/ContextualAnimViewModel.cpp
+++ b/Engine/Plugins/Experimental/Animation/ContextualAnimation/Source/ContextualAnimationEditor/Private/ContextualAnimViewModel.cpp
@@ -544,6 +544,7 @@ void FContextualAnimViewModel::AddNewAnimSet(const FContextualAnimNewAnimSetPara
 		AnimTrack.Role = Data.RoleName;
 		AnimTrack.Animation = Data.Animation;
 		AnimTrack.bRequireFlyingMode = Data.bRequiresFlyingMode;
+		AnimTrack.bOptional = Data.bOptional;
 		AnimSet.Tracks.Add(AnimTrack);
 		AnimSet.RandomWeight = Params.RandomWeight;
 	}
diff --git a/Engine/Plugins/Experimental/Animation/ContextualAnimation/Source/ContextualAnimationEditor/Public/ContextualAnimEditorTypes.h b/Engine/Plugins/Experimental/Animation/ContextualAnimation/Source/ContextualAnimationEditor/Public/ContextualAnimEditorTypes.h
--- a/Engine/Plugins/Experimental/Animation/ContextualAnimation/Source/ContextualAnimationEditor/Public/ContextualAnimEditorTypes.h
+++ b/Engine/Plugins/Experimental/Animation/ContextualAnimation/Source/ContextualAnimationEditor/Public/ContextualAnimEditorTypes.h
@@ -75,6 +75,11 @@ struct FContextualAnimNewAnimSetData
 
 	UPROPERTY(EditAnywhere, Category = "Settings")
 	bool bRequiresFlyingMode = false;
+
+	UPROPERTY(EditAnywhere, Category = "Settings")
+	bool bOptional = false;
+
+	//@TODO: Refactor this to use FContextualAnimTrack directly with a details customization to hide the properties that are not editable when adding a new set
 };
 
 /** Struct used to construct the widget for adding a new set */
