Contextual Anim Editor: Fixed bug where mesh to scene transform could end up with a wrong value after adjusting the position of the actor in the preview scene

diff --git a/Engine/Plugins/Experimental/Animation/ContextualAnimation/Source/ContextualAnimationEditor/Private/ContextualAnimViewModel.cpp b/Engine/Plugins/Experimental/Animation/ContextualAnimation/Source/ContextualAnimationEditor/Private/ContextualAnimViewModel.cpp
--- a/Engine/Plugins/Experimental/Animation/ContextualAnimation/Source/ContextualAnimationEditor/Private/ContextualAnimViewModel.cpp
+++ b/Engine/Plugins/Experimental/Animation/ContextualAnimation/Source/ContextualAnimationEditor/Private/ContextualAnimViewModel.cpp
@@ -1050,7 +1050,7 @@ bool FContextualAnimViewModel::ProcessInputDelta(FVector& InDrag, FRotator& InRo
 
 			return true;
 		}
-		else if (ModifyingActorTransformInSceneState == EModifyActorTransformInSceneState::Modifying)
+		else if (ModifyingActorTransformInSceneState != EModifyActorTransformInSceneState::Inactive)
 		{
 			if (WantsToModifyMeshToSceneForSelectedActor())
 			{
@@ -1157,9 +1157,13 @@ bool FContextualAnimViewModel::StartTracking()
 	{
 		if (const FContextualAnimSceneBinding* Binding = GetSelectedBinding())
 		{
-			NewMeshToSceneTransform = SceneBindings.GetAnimTrackFromBinding(*Binding).MeshToScene;
-			ModifyingActorTransformInSceneState = EModifyActorTransformInSceneState::Modifying;
-			ModifyingTransformInSceneCachedActor = Binding->GetActor();
+			if (ModifyingActorTransformInSceneState == EModifyActorTransformInSceneState::Inactive)
+			{
+				NewMeshToSceneTransform = SceneBindings.GetAnimTrackFromBinding(*Binding).MeshToScene;
+				ModifyingActorTransformInSceneState = EModifyActorTransformInSceneState::Modifying;
+				ModifyingTransformInSceneCachedActor = Binding->GetActor();
+			}
+
 			return true;
 		}
 	}
