Contextual Anim Editor: Downgraded log verbosity

diff --git a/Engine/Plugins/Experimental/Animation/ContextualAnimation/Source/ContextualAnimation/Private/ContextualAnimSelectionCriterion.cpp b/Engine/Plugins/Experimental/Animation/ContextualAnimation/Source/ContextualAnimation/Private/ContextualAnimSelectionCriterion.cpp
--- a/Engine/Plugins/Experimental/Animation/ContextualAnimation/Source/ContextualAnimation/Private/ContextualAnimSelectionCriterion.cpp
+++ b/Engine/Plugins/Experimental/Animation/ContextualAnimation/Source/ContextualAnimation/Private/ContextualAnimSelectionCriterion.cpp
@@ -76,7 +76,7 @@ bool UContextualAnimSelectionCriterion_TriggerArea::DoesQuerierPassCondition(con
 		bResult = (FMath::Abs(AngleSum) > 0.001f);
 	}
 
-	UE_LOG(LogContextualAnim, Verbose, TEXT("UContextualAnimSelectionCriterion_TriggerArea: Primary: %s Querier: %s Result: %d"),
+	UE_LOG(LogContextualAnim, VeryVerbose, TEXT("UContextualAnimSelectionCriterion_TriggerArea: Primary: %s Querier: %s Result: %d"),
 		*GetNameSafe(Primary.GetActor()), *GetNameSafe(Querier.GetActor()), bResult);
 
 	return bResult;
@@ -113,8 +113,8 @@ bool UContextualAnimSelectionCriterion_Cone::DoesQuerierPassCondition(const FCon
 	}
 	
 
-	UE_LOG(LogContextualAnim, Verbose, TEXT("UContextualAnimSelectionCriterion_Cone: Primary: %s Querier: %s Mode: %s Distance: %.1f HalfAngle: %.1f Offset: %.1f Result: %d"), 
-		*GetNameSafe(Primary.GetActor()), *GetNameSafe(Querier.GetActor()), *UEnum::GetValueAsString(TEXT("ContextualAnimation.EContextualAnimCriterionAngleMode"), Mode), Distance, HalfAngle, Offset, bResult);
+	UE_LOG(LogContextualAnim, VeryVerbose, TEXT("UContextualAnimSelectionCriterion_Cone: Primary: %s Querier: %s Mode: %s Distance: %.1f HalfAngle: %.1f Offset: %.1f Result: %d"), 
+		*GetNameSafe(Primary.GetActor()), *GetNameSafe(Querier.GetActor()), *UEnum::GetValueAsString(TEXT("ContextualAnimation.EContextualAnimCriterionConeMode"), Mode), Distance, HalfAngle, Offset, bResult);
 
 	return bResult;
 }
@@ -139,7 +139,7 @@ bool UContextualAnimSelectionCriterion_Distance::DoesQuerierPassCondition(const
 
 	const bool bResult = FMath::IsWithinInclusive(Distance, MinDistance, MaxDistance);
 
-	UE_LOG(LogContextualAnim, Verbose, TEXT("UContextualAnimSelectionCriterion_Distance: Primary: %s Querier: %s Mode: %s MaxDistance: %.1f MaxDist: %.1f Dist: %.1f Result: %d"),
+	UE_LOG(LogContextualAnim, VeryVerbose, TEXT("UContextualAnimSelectionCriterion_Distance: Primary: %s Querier: %s Mode: %s MaxDistance: %.1f MaxDist: %.1f Dist: %.1f Result: %d"),
 		*GetNameSafe(Primary.GetActor()), *GetNameSafe(Querier.GetActor()), *UEnum::GetValueAsString(TEXT("ContextualAnimation.EContextualAnimCriterionDistanceMode"), Mode), MinDistance, MaxDistance, Distance, bResult);
 
 	return bResult;
