Contextual Anim Editor: Removed SceneInstance dependecy and fixed bug where preview playback and timeline could get out of sync

diff --git a/Engine/Plugins/Experimental/Animation/ContextualAnimation/Source/ContextualAnimationEditor/Private/ContextualAnimViewModel.cpp b/Engine/Plugins/Experimental/Animation/ContextualAnimation/Source/ContextualAnimationEditor/Private/ContextualAnimViewModel.cpp
--- a/Engine/Plugins/Experimental/Animation/ContextualAnimation/Source/ContextualAnimationEditor/Private/ContextualAnimViewModel.cpp
+++ b/Engine/Plugins/Experimental/Animation/ContextualAnimation/Source/ContextualAnimationEditor/Private/ContextualAnimViewModel.cpp
@@ -4,6 +4,7 @@
 #include "Modules/ModuleManager.h"
 #include "MovieScene.h"
 #include "ISequencerModule.h"
+#include "MovieSceneSequencePlayer.h"
 #include "MovieSceneCommonHelpers.h"
 #include "ContextualAnimMovieSceneSequence.h"
 #include "ContextualAnimMovieSceneTrack.h"
@@ -22,8 +23,6 @@
 #include "ContextualAnimEditorTypes.h"
 #include "ContextualAnimUtilities.h"
 #include "ContextualAnimSceneAsset.h"
-#include "ContextualAnimManager.h"
-#include "ContextualAnimSceneInstance.h"
 #include "ContextualAnimSceneActorComponent.h"
 #include "AnimNotifyState_IKWindow.h"
 #include "PropertyEditorModule.h"
@@ -50,7 +49,6 @@ FContextualAnimViewModel::FContextualAnimViewModel()
 	: SceneAsset(nullptr)
 	, MovieSceneSequence(nullptr)
 	, MovieScene(nullptr)
-	, ContextualAnimManager(nullptr)
 {
 }
 
@@ -64,7 +62,6 @@ void FContextualAnimViewModel::AddReferencedObjects(FReferenceCollector& Collect
 	Collector.AddReferencedObject(SceneAsset);
 	Collector.AddReferencedObject(MovieSceneSequence);
 	Collector.AddReferencedObject(MovieScene);
-	Collector.AddReferencedObject(ContextualAnimManager);
 }
 
 TSharedPtr<ISequencer> FContextualAnimViewModel::GetSequencer()
@@ -79,8 +76,6 @@ void FContextualAnimViewModel::Initialize(UContextualAnimSceneAsset* InSceneAsse
 	SceneAsset = InSceneAsset;
 	PreviewScenePtr = InPreviewScene;
 
-	ContextualAnimManager = NewObject<UContextualAnimManager>(GetWorld());
-
 	CreateSequencer();
 
 	CreateDetailsView();
@@ -474,19 +469,14 @@ void FContextualAnimViewModel::SetNotifiesMode(const FContextualAnimTrack& AnimT
 
 void FContextualAnimViewModel::RefreshPreviewScene()
 {
-	if (SceneInstance.IsValid())
-	{
-		SceneInstance->GetBindings().Reset();
-		ContextualAnimManager->OnSceneInstanceEnded(SceneInstance.Get());
-		SceneInstance.Reset();
-	}
+	SceneBindings.Reset();
 
 	// Remove preview actors from the preview scene
 	// @TODO: Investigate why GetActor from each binding in SceneBindings returns null after a Redo op
 	for (TActorIterator<AActor> It(GetWorld()); It; ++It)
 	{
 		AActor* Actor = *It;
-		if(Actor && Actor->ActorHasTag(PreviewActorTag))
+		if (Actor && Actor->ActorHasTag(PreviewActorTag))
 		{
 			Actor->Destroy();
 		}
@@ -503,27 +493,22 @@ void FContextualAnimViewModel::RefreshPreviewScene()
 	{
 		return;
 	}
-	
-	FContextualAnimStartSceneParams StartSceneParams;
-	StartSceneParams.SectionIdx = ActiveSectionIdx;
-	StartSceneParams.AnimSetIdx = ActiveAnimSetIdx;
+
+	SceneBindings = FContextualAnimSceneBindings(*SceneAsset, ActiveSectionIdx, ActiveAnimSetIdx);
 
 	const FContextualAnimSet& AnimSet = ContextualAnimSection.AnimSets[ActiveAnimSetIdx];
 	for (const FContextualAnimTrack& AnimTrack : AnimSet.Tracks)
 	{
 		if (AActor* PreviewActor = SpawnPreviewActor(AnimTrack))
 		{
-			StartSceneParams.RoleToActorMap.Add(AnimTrack.Role, PreviewActor);
+			SceneBindings.BindActorToRole(*PreviewActor, AnimTrack.Role);
 		}
 	}
 
-	SceneInstance = ContextualAnimManager->ForceStartScene(*GetSceneAsset(), StartSceneParams);
-
-	// Disable auto blend out
-	if (SceneInstance.IsValid())
+	// Pause all anims and disable auto blend out
+	const FContextualAnimSceneBinding* Leader = SceneBindings.GetSyncLeader();
+	if (Leader && Leader->GetSceneActorComponent()->StartContextualAnimScene(SceneBindings))
 	{
-		SceneBindings = SceneInstance->GetBindings();
-
 		for (const FContextualAnimSceneBinding& Binding : SceneBindings)
 		{
 			if (FAnimMontageInstance* MontageInstance = Binding.GetAnimMontageInstance())
@@ -654,40 +639,48 @@ void FContextualAnimViewModel::AddNewIKTarget(const UContextualAnimNewIKTargetPa
 	SceneAsset->MarkPackageDirty();
 }
 
-void FContextualAnimViewModel::ToggleSimulateMode()
+void FContextualAnimViewModel::StartSimulateMode()
 {
-	if (SimulateModeState == ESimulateModeState::Inactive)
-	{
-		MovieScene->SetReadOnly(true);
+	Sequencer->Pause();
+	Sequencer->SetGlobalTime(0);
 
-		SimulateModeState = ESimulateModeState::Paused;
+	MovieScene->SetReadOnly(true);
 
-		if (SceneInstance.IsValid())
-		{
-			SceneInstance->Stop();
-		}
-	}
-	else
-	{
-		MovieScene->SetReadOnly(false);
+	SimulateModeState = ESimulateModeState::Paused;
+}
+
+void FContextualAnimViewModel::StopSimulateMode()
+{
+	SimulateModeState = ESimulateModeState::Inactive;
+
+	Sequencer->Pause();
+	Sequencer->SetGlobalTime(0);
 
-		SimulateModeState = ESimulateModeState::Inactive;
+	MovieScene->SetReadOnly(false);
 
-		if (SceneInstance.IsValid())
+	for (auto& Binding : SceneBindings)
+	{
+		if (UMotionWarpingComponent* MotionWarpComp = Binding.GetActor()->FindComponentByClass<UMotionWarpingComponent>())
 		{
-			for (auto& Binding : SceneInstance->GetBindings())
+			for (const FContextualAnimSetPivotDefinition& Def : SceneAsset->GetAnimSetPivotDefinitionsInSection(SceneBindings.GetSectionIdx()))
 			{
-				if (UMotionWarpingComponent* MotionWarpComp = Binding.GetActor()->FindComponentByClass<UMotionWarpingComponent>())
-				{
-					for (const FContextualAnimSetPivotDefinition& Def : SceneAsset->GetAnimSetPivotDefinitionsInSection(SceneInstance->GetBindings().GetSectionIdx()))
-					{
-						MotionWarpComp->RemoveWarpTarget(Def.Name);
-					}
-				}
+				MotionWarpComp->RemoveWarpTarget(Def.Name);
 			}
 		}
+	}
 
-		SetDefaultMode();
+	SetDefaultMode();
+}
+
+void FContextualAnimViewModel::ToggleSimulateMode()
+{
+	if (SimulateModeState == ESimulateModeState::Inactive)
+	{
+		StartSimulateMode();
+	}
+	else
+	{
+		StopSimulateMode();
 	}
 
 	ClearSelection();
@@ -695,24 +688,20 @@ void FContextualAnimViewModel::ToggleSimulateMode()
 
 void FContextualAnimViewModel::StartSimulation()
 {
-	FContextualAnimStartSceneParams Params;
+	TMap<FName, FContextualAnimSceneBindingContext> RoleToActorMap;
 	for (const FContextualAnimSceneBinding& Binding : SceneBindings)
 	{
-		Params.RoleToActorMap.Add(SceneBindings.GetRoleFromBinding(Binding), Binding.GetContext());
+		RoleToActorMap.Add(SceneBindings.GetRoleFromBinding(Binding), Binding.GetContext());
 	}
-	Params.SectionIdx = 0;
-	SceneInstance = ContextualAnimManager->TryStartScene(*GetSceneAsset(), Params);
 
-	if(SceneInstance == nullptr)
+	if (FContextualAnimSceneBindings::TryCreateBindings(*GetSceneAsset(), ActiveSectionIdx, RoleToActorMap, SceneBindings))
 	{
-		//@TODO: This should be a message on the screen
-		UE_LOG(LogContextualAnim, Warning, TEXT("Can't start scene"));
-		return;
+		const FContextualAnimSceneBinding* Leader = SceneBindings.GetSyncLeader();
+		if (Leader && Leader->GetSceneActorComponent()->StartContextualAnimScene(SceneBindings))
+		{
+			SimulateModeState = ESimulateModeState::Playing;
+		}
 	}
-
-	SceneBindings = SceneInstance->GetBindings();
-
-	SimulateModeState = ESimulateModeState::Playing;
 }
 
 UWorld* FContextualAnimViewModel::GetWorld() const
@@ -797,51 +786,30 @@ void FContextualAnimViewModel::SequencerStopEvent()
 {
 	if (SimulateModeState != ESimulateModeState::Inactive)
 	{
-		ToggleSimulateMode();
+		StopSimulateMode();
 	}
 }
 
 void FContextualAnimViewModel::SequencerTimeChanged()
 {
+	if (!IsSimulateModeInactive())
+	{
+		PreviousSequencerStatus = EMovieScenePlayerStatus::Stopped;
+		PreviousSequencerTime = 0.f;
+		return;
+	}
+
 	EMovieScenePlayerStatus::Type CurrentStatus = Sequencer->GetPlaybackStatus();
 	const float CurrentSequencerTime = Sequencer->GetGlobalTime().AsSeconds();
-	const float PlaybackSpeed = Sequencer->GetPlaybackSpeed();
 
-	if (SceneInstance.IsValid())
+	for (const FContextualAnimSceneBinding& Binding : SceneBindings)
 	{
-		for (auto& Binding : SceneInstance->GetBindings())
+		if (FAnimMontageInstance* MontageInstance = Binding.GetAnimMontageInstance())
 		{
-			if (FAnimMontageInstance* MontageInstance = Binding.GetAnimMontageInstance())
-			{
-				const float AnimPlayLength = MontageInstance->Montage->GetPlayLength();
-				float PreviousTime = FMath::Clamp(PreviousSequencerTime, 0.f, AnimPlayLength);
-				float CurrentTime = FMath::Clamp(CurrentSequencerTime, 0.f, AnimPlayLength);
-
-				if (CurrentStatus == EMovieScenePlayerStatus::Stopped || CurrentStatus == EMovieScenePlayerStatus::Scrubbing)
-				{
-					UpdatePreviewActorTransform(Binding, CurrentTime);
-
-					if (MontageInstance->IsPlaying())
-					{
-						MontageInstance->Pause();
-					}
-
-					MontageInstance->SetPosition(CurrentTime);
-				}
-				else if (CurrentStatus == EMovieScenePlayerStatus::Playing)
-				{
-					if (PlaybackSpeed > 0.f && CurrentTime < PreviousTime)
-					{
-						UpdatePreviewActorTransform(Binding, CurrentTime);
-						MontageInstance->SetPosition(CurrentTime);
-					}
-
-					if (!MontageInstance->IsPlaying())
-					{
-						MontageInstance->SetPlaying(true);
-					}
-				}
-			}
+			const float AnimPlayLength = MontageInstance->Montage->GetPlayLength();
+			const float CurrentTime = FMath::Clamp(CurrentSequencerTime, 0.f, AnimPlayLength - KINDA_SMALL_NUMBER);
+			MontageInstance->SetPosition(CurrentTime);
+			UpdatePreviewActorTransform(Binding, CurrentTime);
 		}
 	}
 
@@ -861,11 +829,6 @@ void FContextualAnimViewModel::SequencerDataChanged(EMovieSceneDataChangeType Da
 		return;
 	}
 
-	if(!SceneInstance.IsValid())
-	{
-		return;
-	}
-
 	if(DataChangeType == EMovieSceneDataChangeType::TrackValueChanged)
 	{
 		// Update IK AnimNotify's bEnable flag based on the Active state of the section
@@ -1073,15 +1036,12 @@ bool FContextualAnimViewModel::GetCustomDrawingCoordinateSystem(FMatrix& InMatri
 {
 	if (SelectionInfo.Criterion.Value != INDEX_NONE)
 	{
-		if (SceneInstance.IsValid())
+		FTransform PrimaryActorTransform = FTransform::Identity;
+		if (const FContextualAnimSceneBinding* Binding = SceneBindings.FindBindingByRole(GetSceneAsset()->GetPrimaryRole()))
 		{
-			FTransform PrimaryActorTransform = FTransform::Identity;
-			if (const FContextualAnimSceneBinding* Binding = SceneInstance->FindBindingByRole(GetSceneAsset()->GetPrimaryRole()))
-			{
-				PrimaryActorTransform = Binding->GetTransform();
-				InMatrix = PrimaryActorTransform.ToMatrixNoScale().RemoveTranslation();
-				return true;
-			}
+			PrimaryActorTransform = Binding->GetTransform();
+			InMatrix = PrimaryActorTransform.ToMatrixNoScale().RemoveTranslation();
+			return true;
 		}
 	}
 	else if (AActor* SelectedActor = GetSelectedActor())
@@ -1135,7 +1095,7 @@ FVector FContextualAnimViewModel::GetWidgetLocationFromSelection() const
 			}
 
 			FTransform PrimaryActorTransform = FTransform::Identity;
-			if (const FContextualAnimSceneBinding* Binding = SceneInstance->FindBindingByRole(SceneAsset->GetPrimaryRole()))
+			if (const FContextualAnimSceneBinding* Binding = SceneBindings.FindBindingByRole(SceneAsset->GetPrimaryRole()))
 			{
 				PrimaryActorTransform = Binding->GetTransform();
 			}
@@ -1220,7 +1180,7 @@ void FContextualAnimViewModel::DiscardChangeToActorTransformInScene()
 {
 	ModifyingActorTransformInSceneState = EModifyActorTransformInSceneState::Inactive;
 
-	if (const FContextualAnimSceneBinding* Binding = SceneInstance->FindBindingByActor(ModifyingTransformInSceneCachedActor.Get()))
+	if (const FContextualAnimSceneBinding* Binding = SceneBindings.FindBindingByActor(ModifyingTransformInSceneCachedActor.Get()))
 	{
 		UpdatePreviewActorTransform(*Binding, Sequencer->GetGlobalTime().AsSeconds());
 	}
diff --git a/Engine/Plugins/Experimental/Animation/ContextualAnimation/Source/ContextualAnimationEditor/Public/ContextualAnimViewModel.h b/Engine/Plugins/Experimental/Animation/ContextualAnimation/Source/ContextualAnimationEditor/Public/ContextualAnimViewModel.h
--- a/Engine/Plugins/Experimental/Animation/ContextualAnimation/Source/ContextualAnimationEditor/Public/ContextualAnimViewModel.h
+++ b/Engine/Plugins/Experimental/Animation/ContextualAnimation/Source/ContextualAnimationEditor/Public/ContextualAnimViewModel.h
@@ -15,7 +15,6 @@ class FContextualAnimPreviewScene;
 class UContextualAnimSceneAsset;
 class UContextualAnimManager;
 class UContextualAnimMovieSceneSequence;
-class UContextualAnimSceneInstance;
 class UMovieScene;
 class UMovieSceneTrack;
 class UMovieSceneSection;
@@ -78,7 +77,6 @@ public:
 	UMovieScene* GetMovieScene() const { return MovieScene; }
 	UContextualAnimMovieSceneSequence* GetMovieSceneSequence() const { return MovieSceneSequence; }
 	UContextualAnimSceneAsset* GetSceneAsset() const { return SceneAsset; }
-	UContextualAnimSceneInstance* GetSceneInstance() const { return SceneInstance.Get(); }
 	UContextualAnimMovieSceneTrack* FindTrackByRole(const FName& Role) const;
 	TSharedPtr<IDetailsView> GetDetailsView() { return DetailsView; }
 
@@ -95,6 +93,8 @@ public:
 	void SetActiveAnimSetForSection(int32 SectionIdx, int32 AnimSetIdx);
 	int32 GetActiveAnimSetForSection(int32 SectionIdx) const;
 
+	void StartSimulateMode();
+	void StopSimulateMode();
 	void ToggleSimulateMode();
 	bool IsSimulateModeInactive()	const { return SimulateModeState == ESimulateModeState::Inactive; }
 	bool IsSimulateModePaused()		const { return SimulateModeState == ESimulateModeState::Paused; }
@@ -170,10 +170,6 @@ private:
 	/** Weak pointer to the PreviewScene */
 	TWeakPtr<FContextualAnimPreviewScene> PreviewScenePtr;
 
-	TObjectPtr<UContextualAnimManager> ContextualAnimManager;
-
-	TWeakObjectPtr<UContextualAnimSceneInstance> SceneInstance;
-
 	/** Copy of the bindings so we can access them even when simulation mode is not playing (and SceneInstance is not valid) */
 	FContextualAnimSceneBindings SceneBindings;
 
