Contextual Anim Editor: Temporally changed animation field in the new anim set dialog back to accept only AnimMontage to prevent an issue that seems to be caused by a recent change somewhere else that prevents this field from showing up on the ui preventing the user from adding data to an scene asset.

diff --git a/Engine/Plugins/Experimental/Animation/ContextualAnimation/Source/ContextualAnimationEditor/Public/ContextualAnimEditorTypes.h b/Engine/Plugins/Experimental/Animation/ContextualAnimation/Source/ContextualAnimationEditor/Public/ContextualAnimEditorTypes.h
--- a/Engine/Plugins/Experimental/Animation/ContextualAnimation/Source/ContextualAnimationEditor/Public/ContextualAnimEditorTypes.h
+++ b/Engine/Plugins/Experimental/Animation/ContextualAnimation/Source/ContextualAnimationEditor/Public/ContextualAnimEditorTypes.h
@@ -71,7 +71,7 @@ struct FContextualAnimNewAnimSetData
 	FName RoleName = NAME_None;
 
 	UPROPERTY(EditAnywhere, Category = "Settings")
-	TObjectPtr<UAnimSequenceBase> Animation = nullptr;
+	TObjectPtr<class UAnimMontage> Animation = nullptr;
 
 	UPROPERTY(EditAnywhere, Category = "Settings")
 	bool bRequiresFlyingMode = false;
