Contextual Anim Editor: Disable resizing or moving around anim sections and set both notify track and anim track to support multiple rows.

diff --git a/Engine/Plugins/Experimental/Animation/ContextualAnimation/Source/ContextualAnimationEditor/Private/ContextualAnimMovieSceneTrack.cpp b/Engine/Plugins/Experimental/Animation/ContextualAnimation/Source/ContextualAnimationEditor/Private/ContextualAnimMovieSceneTrack.cpp
--- a/Engine/Plugins/Experimental/Animation/ContextualAnimation/Source/ContextualAnimationEditor/Private/ContextualAnimMovieSceneTrack.cpp
+++ b/Engine/Plugins/Experimental/Animation/ContextualAnimation/Source/ContextualAnimationEditor/Private/ContextualAnimMovieSceneTrack.cpp
@@ -4,6 +4,7 @@
 #include "ContextualAnimMovieSceneSection.h"
 #include "ContextualAnimMovieSceneSequence.h"
 #include "MovieScene.h"
+#include "ContextualAnimTypes.h"
 
 #include UE_INLINE_GENERATED_CPP_BY_NAME(ContextualAnimMovieSceneTrack)
 
@@ -71,6 +72,19 @@ void UContextualAnimMovieSceneTrack::RemoveSectionAt(int32 SectionIndex)
 #if WITH_EDITOR
 EMovieSceneSectionMovedResult UContextualAnimMovieSceneTrack::OnSectionMoved(UMovieSceneSection& Section, const FMovieSceneSectionMovedParams& Params)
 {
+	if (Params.MoveType == EPropertyChangeType::ValueSet)
+	{
+		UContextualAnimMovieSceneSection* NotifySection = CastChecked<UContextualAnimMovieSceneSection>(&Section);
+
+		const FContextualAnimTrack& AnimTrack = NotifySection->GetAnimTrack();
+		if(AnimTrack.Animation)
+		{
+			const float AnimLength = AnimTrack.Animation->GetPlayLength();
+			const FFrameRate TickResolution = GetTypedOuter<UMovieScene>()->GetTickResolution();
+			NotifySection->SetRange(TRange<FFrameNumber>::Inclusive(0, (AnimLength * TickResolution).RoundToFrame()));
+		}
+	}
+
 	return EMovieSceneSectionMovedResult::None;
 }
 #endif
diff --git a/Engine/Plugins/Experimental/Animation/ContextualAnimation/Source/ContextualAnimationEditor/Public/ContextualAnimMovieSceneNotifyTrack.h b/Engine/Plugins/Experimental/Animation/ContextualAnimation/Source/ContextualAnimationEditor/Public/ContextualAnimMovieSceneNotifyTrack.h
--- a/Engine/Plugins/Experimental/Animation/ContextualAnimation/Source/ContextualAnimationEditor/Public/ContextualAnimMovieSceneNotifyTrack.h
+++ b/Engine/Plugins/Experimental/Animation/ContextualAnimation/Source/ContextualAnimationEditor/Public/ContextualAnimMovieSceneNotifyTrack.h
@@ -56,6 +56,8 @@ public:
 	virtual EMovieSceneSectionMovedResult OnSectionMoved(UMovieSceneSection& Section, const FMovieSceneSectionMovedParams& Params) override;
 #endif
 
+	virtual bool SupportsMultipleRows() const override { return true; }
+
 	/** Returns a reference to the animation this track was created from */
 	UAnimSequenceBase& GetAnimation() const;
 
diff --git a/Engine/Plugins/Experimental/Animation/ContextualAnimation/Source/ContextualAnimationEditor/Public/ContextualAnimMovieSceneTrack.h b/Engine/Plugins/Experimental/Animation/ContextualAnimation/Source/ContextualAnimationEditor/Public/ContextualAnimMovieSceneTrack.h
--- a/Engine/Plugins/Experimental/Animation/ContextualAnimation/Source/ContextualAnimationEditor/Public/ContextualAnimMovieSceneTrack.h
+++ b/Engine/Plugins/Experimental/Animation/ContextualAnimation/Source/ContextualAnimationEditor/Public/ContextualAnimMovieSceneTrack.h
@@ -52,6 +52,8 @@ public:
 	virtual EMovieSceneSectionMovedResult OnSectionMoved(UMovieSceneSection& Section, const FMovieSceneSectionMovedParams& Params) override;
 #endif
 
+	virtual bool SupportsMultipleRows() const override { return true; }
+
 private:
 
 	/** List of sections in this track */
diff --git a/Engine/Plugins/Experimental/Animation/ContextualAnimation/Source/ContextualAnimationEditor/Public/ContextualAnimMovieSceneTrackEditor.h b/Engine/Plugins/Experimental/Animation/ContextualAnimation/Source/ContextualAnimationEditor/Public/ContextualAnimMovieSceneTrackEditor.h
--- a/Engine/Plugins/Experimental/Animation/ContextualAnimation/Source/ContextualAnimationEditor/Public/ContextualAnimMovieSceneTrackEditor.h
+++ b/Engine/Plugins/Experimental/Animation/ContextualAnimation/Source/ContextualAnimationEditor/Public/ContextualAnimMovieSceneTrackEditor.h
@@ -39,4 +39,5 @@ public:
 	virtual FReply OnSectionDoubleClicked(const FGeometry& SectionGeometry, const FPointerEvent& MouseEvent) override;
 	virtual FText GetSectionTitle() const override;
 	virtual float GetSectionHeight() const override;
+	virtual bool SectionIsResizable() const override { return false; }
 };
\ No newline at end of file
