Contextual Anim Editor: Fixed bug where anim notify's details was not appearing on the context menu/properties of the corresponding section.

diff --git a/Engine/Plugins/Experimental/Animation/ContextualAnimation/Source/ContextualAnimationEditor/Private/DetailCustomizations/ContextualAnimNotifySectionDetailCustom.cpp b/Engine/Plugins/Experimental/Animation/ContextualAnimation/Source/ContextualAnimationEditor/Private/DetailCustomizations/ContextualAnimNotifySectionDetailCustom.cpp
--- a/Engine/Plugins/Experimental/Animation/ContextualAnimation/Source/ContextualAnimationEditor/Private/DetailCustomizations/ContextualAnimNotifySectionDetailCustom.cpp
+++ b/Engine/Plugins/Experimental/Animation/ContextualAnimation/Source/ContextualAnimationEditor/Private/DetailCustomizations/ContextualAnimNotifySectionDetailCustom.cpp
@@ -4,6 +4,7 @@
 #include "Animation/AnimTypes.h"
 #include "DetailLayoutBuilder.h"
 #include "DetailCategoryBuilder.h"
+#include "Animation/AnimNotifies/AnimNotify.h"
 #include "Animation/AnimNotifies/AnimNotifyState.h"
 #include "ContextualAnimMovieSceneNotifySection.h"
 #include "IDetailPropertyRow.h"
@@ -25,11 +26,21 @@ void FContextualAnimNotifySectionDetailCustom::CustomizeDetails(IDetailLayoutBui
 	// Add new category to show the properties from the actual AnimNotify object
 	if (const FAnimNotifyEvent* NotifyEvent = Section->GetAnimNotifyEvent())
 	{
-		if (UAnimNotifyState* NotifyState = NotifyEvent->NotifyStateClass)
+		UObject* NotifyObject = nullptr;
+		if (NotifyEvent->Notify)
+		{
+			NotifyObject = NotifyEvent->Notify;
+		}
+		else
+		{
+			NotifyObject = NotifyEvent->NotifyStateClass;
+		}
+
+		if (NotifyObject)
 		{
 			IDetailCategoryBuilder& Category = DetailBuilder.EditCategory(TEXT("Notify"), FText::GetEmpty(), ECategoryPriority::Variable);
 
-			TArray<UObject*> ExternalObjects = { NotifyState };
+			TArray<UObject*> ExternalObjects = { NotifyObject };
 			IDetailPropertyRow* PropertyRow = Category.AddExternalObjects(ExternalObjects);
 			PropertyRow->ShouldAutoExpand(true);
 		}
