Contextual Anim Editor: Added option to specify a preview animinstance for any of the actors in the scene

diff --git a/Engine/Plugins/Experimental/Animation/ContextualAnimation/Source/ContextualAnimation/Public/ContextualAnimSceneAsset.h b/Engine/Plugins/Experimental/Animation/ContextualAnimation/Source/ContextualAnimation/Public/ContextualAnimSceneAsset.h
--- a/Engine/Plugins/Experimental/Animation/ContextualAnimation/Source/ContextualAnimation/Public/ContextualAnimSceneAsset.h
+++ b/Engine/Plugins/Experimental/Animation/ContextualAnimation/Source/ContextualAnimation/Public/ContextualAnimSceneAsset.h
@@ -174,6 +174,9 @@ struct FContextualAnimActorPreviewData
 	UPROPERTY(EditAnywhere, Category = "Defaults", meta = (EditCondition = "Type==EContextualAnimActorPreviewType::SkeletalMesh", EditConditionHides))
 	TSoftObjectPtr<class USkeletalMesh> PreviewSkeletalMesh;
 
+	UPROPERTY(EditAnywhere, Category = "Defaults", meta = (EditCondition = "Type==EContextualAnimActorPreviewType::SkeletalMesh", EditConditionHides))
+	TSoftClassPtr<class UAnimInstance> PreviewAnimInstance;
+
 	UPROPERTY(EditAnywhere, Category = "Defaults", meta = (EditCondition = "Type==EContextualAnimActorPreviewType::StaticMesh", EditConditionHides))
 	TSoftObjectPtr<class UStaticMesh> PreviewStaticMesh;
 
diff --git a/Engine/Plugins/Experimental/Animation/ContextualAnimation/Source/ContextualAnimationEditor/Private/ContextualAnimViewModel.cpp b/Engine/Plugins/Experimental/Animation/ContextualAnimation/Source/ContextualAnimationEditor/Private/ContextualAnimViewModel.cpp
--- a/Engine/Plugins/Experimental/Animation/ContextualAnimation/Source/ContextualAnimationEditor/Private/ContextualAnimViewModel.cpp
+++ b/Engine/Plugins/Experimental/Animation/ContextualAnimation/Source/ContextualAnimationEditor/Private/ContextualAnimViewModel.cpp
@@ -190,14 +190,15 @@ AActor* FContextualAnimViewModel::SpawnPreviewActor(const FContextualAnimTrack&
 	USkeletalMesh* PreviewSkeletalMesh = nullptr;
 	UStaticMesh* PreviewStaticMesh = nullptr;
 	UClass* PreviewActorClass = nullptr;
+	UClass* PreviewAnimInstanceClass = nullptr;
 
 	for (const FContextualAnimActorPreviewData& PreviewData : GetSceneAsset()->OverridePreviewData)
 	{
 		if (PreviewData.Role == AnimTrack.Role)
 		{
-			if(PreviewData.Type == EContextualAnimActorPreviewType::StaticMesh)
+			if (PreviewData.Type == EContextualAnimActorPreviewType::StaticMesh)
 			{
-				if(UStaticMesh* StaticMesh = PreviewData.PreviewStaticMesh.LoadSynchronous())
+				if (UStaticMesh* StaticMesh = PreviewData.PreviewStaticMesh.LoadSynchronous())
 				{
 					PreviewStaticMesh = StaticMesh;
 					PreviewType = PreviewData.Type;
@@ -210,6 +211,11 @@ AActor* FContextualAnimViewModel::SpawnPreviewActor(const FContextualAnimTrack&
 					PreviewSkeletalMesh = SkeletalMesh;
 					PreviewType = PreviewData.Type;
 				}
+
+				if (UClass* AnimInstanceClass = PreviewData.PreviewAnimInstance.LoadSynchronous())
+				{
+					PreviewAnimInstanceClass = AnimInstanceClass;
+				}
 			}
 			else if (PreviewData.Type == EContextualAnimActorPreviewType::Actor)
 			{
@@ -254,11 +260,21 @@ AActor* FContextualAnimViewModel::SpawnPreviewActor(const FContextualAnimTrack&
 			SkelMeshComp->SetRelativeRotation(RoleDef->MeshToComponent.GetRotation());
 			SkelMeshComp->SetSkeletalMesh(PreviewSkeletalMesh);
 			SkelMeshComp->SetAnimationMode(EAnimationMode::AnimationBlueprint);
-			SkelMeshComp->SetAnimInstanceClass(UAnimPreviewInstance::StaticClass());
 
-			UAnimPreviewInstance* AnimInstance = Cast<UAnimPreviewInstance>(SkelMeshComp->GetAnimInstance());
-			AnimInstance->SetAnimationAsset(AnimTrack.Animation, false, 0.f);
-			AnimInstance->PlayAnim(false, 0.f);
+			if (PreviewAnimInstanceClass)
+			{
+				SkelMeshComp->SetAnimInstanceClass(PreviewAnimInstanceClass);
+			}
+			else
+			{
+				SkelMeshComp->SetAnimInstanceClass(UAnimPreviewInstance::StaticClass());
+
+				if (UAnimPreviewInstance* AnimInstance = Cast<UAnimPreviewInstance>(SkelMeshComp->GetAnimInstance()))
+				{
+					AnimInstance->SetAnimationAsset(AnimTrack.Animation, false, 0.f);
+					AnimInstance->PlayAnim(false, 0.f);
+				}
+			}
 
 			PreviewCharacter->CacheInitialMeshOffset(SkelMeshComp->GetRelativeLocation(), SkelMeshComp->GetRelativeRotation());
 
@@ -280,12 +296,12 @@ AActor* FContextualAnimViewModel::SpawnPreviewActor(const FContextualAnimTrack&
 		}
 		else
 		{
-			if(PreviewActorClass)
+			if (PreviewActorClass)
 			{
 				PreviewActor = GetWorld()->SpawnActor<AActor>(PreviewActorClass, SpawnTransform, Params);
 				PreviewActor->SetFlags(RF_Transient);
 			}
-			else 
+			else
 			{
 				PreviewActor = GetWorld()->SpawnActor<AActor>(AActor::StaticClass(), SpawnTransform, Params);
 				PreviewActor->SetFlags(RF_Transient);
@@ -293,19 +309,25 @@ AActor* FContextualAnimViewModel::SpawnPreviewActor(const FContextualAnimTrack&
 				if (PreviewSkeletalMesh)
 				{
 					UDebugSkelMeshComponent* SkelMeshComp = NewObject<UDebugSkelMeshComponent>(PreviewActor);
-					SkelMeshComp->RegisterComponentWithWorld(GetWorld());
-
-					UAnimPreviewInstance* AnimInstance = NewObject<UAnimPreviewInstance>(SkelMeshComp);
-					SkelMeshComp->PreviewInstance = AnimInstance;
-					AnimInstance->InitializeAnimation();
-
 					SkelMeshComp->SetSkeletalMesh(PreviewSkeletalMesh);
-					SkelMeshComp->EnablePreview(true, AnimTrack.Animation);
+					SkelMeshComp->SetAnimationMode(EAnimationMode::AnimationBlueprint);
 
-					AnimInstance->SetAnimationAsset(AnimTrack.Animation, false, 0.0f);
+					if (PreviewAnimInstanceClass)
+					{
+						SkelMeshComp->SetAnimInstanceClass(PreviewAnimInstanceClass);
+					}
+					else
+					{
+						SkelMeshComp->SetAnimInstanceClass(UAnimPreviewInstance::StaticClass());
 
-					AnimInstance->PlayAnim(false, 0.0f);
+						if (UAnimPreviewInstance* AnimInstance = Cast<UAnimPreviewInstance>(SkelMeshComp->GetAnimInstance()))
+						{
+							AnimInstance->SetAnimationAsset(AnimTrack.Animation, false, 0.f);
+							AnimInstance->PlayAnim(false, 0.f);
+						}
+					}
 
+					SkelMeshComp->RegisterComponentWithWorld(GetWorld());
 					PreviewActor->SetRootComponent(SkelMeshComp);
 				}
 				else if (PreviewStaticMesh)
@@ -523,6 +545,23 @@ void FContextualAnimViewModel::AddNewAnimSet(const FContextualAnimNewAnimSetPara
 		AnimTrack.bOptional = Data.bOptional;
 		AnimSet.Tracks.Add(AnimTrack);
 		AnimSet.RandomWeight = Params.RandomWeight;
+
+		// Update preview mesh if necessary
+		if (Data.Animation && AnimTrack.Animation->GetSkeleton())
+		{
+			const FName Role = Data.RoleName;
+			if (!SceneAsset->OverridePreviewData.ContainsByPredicate([Role](const FContextualAnimActorPreviewData& Item) { return Item.Role == Role; }))
+			{
+				if (USkeletalMesh* SkelMesh = AnimTrack.Animation->GetSkeleton()->GetPreviewMesh(true))
+				{
+					FContextualAnimActorPreviewData NewPreviewData;
+					NewPreviewData.Role = Role;
+					NewPreviewData.Type = EContextualAnimActorPreviewType::SkeletalMesh;
+					NewPreviewData.PreviewSkeletalMesh = SkelMesh;
+					SceneAsset->OverridePreviewData.Add(NewPreviewData);
+				}
+			}
+		}
 	}
 
 	int32 SectionIdx = INDEX_NONE;
@@ -1195,10 +1234,6 @@ void FContextualAnimViewModel::OnFinishedChangingProperties(const FPropertyChang
 	{
 		SetDefaultMode();
 	}
-	else if (MemberPropertyName == GET_MEMBER_NAME_CHECKED(UContextualAnimSceneAsset, OverridePreviewData))
-	{
-		SetDefaultMode();
-	}
 }
 
 bool FContextualAnimViewModel::MatchesContext(const FTransactionContext& InContext, const TArray<TPair<UObject*, FTransactionObjectEvent>>& TransactionObjectContexts) const
