Contextual Anim: Refactored cached components into weak object ptrs and marked as uproperties + added safe check to ensure scene actor comp is valid before calling joinscene

diff --git a/Engine/Plugins/Experimental/Animation/ContextualAnimation/Source/ContextualAnimation/Private/ContextualAnimSceneActorComponent.cpp b/Engine/Plugins/Experimental/Animation/ContextualAnimation/Source/ContextualAnimation/Private/ContextualAnimSceneActorComponent.cpp
--- a/Engine/Plugins/Experimental/Animation/ContextualAnimation/Source/ContextualAnimation/Private/ContextualAnimSceneActorComponent.cpp
+++ b/Engine/Plugins/Experimental/Animation/ContextualAnimation/Source/ContextualAnimation/Private/ContextualAnimSceneActorComponent.cpp
@@ -71,10 +71,10 @@ bool UContextualAnimSceneActorComponent::StartContextualAnimScene(const FContext
 			{
 				if (Binding.GetActor() != GetOwner())
 				{
-					UContextualAnimSceneActorComponent* Comp = Binding.GetSceneActorComponent();
-					checkf(Comp, TEXT("Missing SceneActorComp on %s"), *GetNameSafe(Binding.GetActor()));
-
-					Comp->JoinScene(InBindings);
+					if (UContextualAnimSceneActorComponent* Comp = Binding.GetSceneActorComponent())
+					{
+						Comp->JoinScene(InBindings);
+					}
 				}
 			}
 
@@ -194,10 +194,10 @@ void UContextualAnimSceneActorComponent::OnRep_Bindings(const FContextualAnimSce
 			{
 				if (Binding.GetActor() != GetOwner())
 				{
-					UContextualAnimSceneActorComponent* Comp = Binding.GetSceneActorComponent();
-					checkf(Comp, TEXT("Missing SceneActorComp on %s"), *GetNameSafe(Binding.GetActor()));
-
-					Comp->JoinScene(RepBindings);
+					if (UContextualAnimSceneActorComponent* Comp = Binding.GetSceneActorComponent())
+					{
+						Comp->JoinScene(RepBindings);
+					}
 				}
 			}
 		}
diff --git a/Engine/Plugins/Experimental/Animation/ContextualAnimation/Source/ContextualAnimation/Private/ContextualAnimTypes.cpp b/Engine/Plugins/Experimental/Animation/ContextualAnimation/Source/ContextualAnimation/Private/ContextualAnimTypes.cpp
--- a/Engine/Plugins/Experimental/Animation/ContextualAnimation/Source/ContextualAnimation/Private/ContextualAnimTypes.cpp
+++ b/Engine/Plugins/Experimental/Animation/ContextualAnimation/Source/ContextualAnimation/Private/ContextualAnimTypes.cpp
@@ -286,7 +286,7 @@ void FContextualAnimSceneBinding::SetAnimTrack(const FContextualAnimTrack& InAni
 
 UContextualAnimSceneActorComponent* FContextualAnimSceneBinding::GetSceneActorComponent() const
 {
-	if (CachedSceneActorComp == nullptr)
+	if (!CachedSceneActorComp.IsValid())
 	{
 		if (AActor* Actor = Context.GetActor())
 		{
@@ -294,27 +294,27 @@ UContextualAnimSceneActorComponent* FContextualAnimSceneBinding::GetSceneActorCo
 		}
 	}
 
-	return CachedSceneActorComp;
+	return CachedSceneActorComp.Get();
 }
 
 UAnimInstance* FContextualAnimSceneBinding::GetAnimInstance() const
 {
-	if (CachedAnimInstance == nullptr)
+	if (!CachedAnimInstance.IsValid())
 	{
 		CachedAnimInstance = UContextualAnimUtilities::TryGetAnimInstance(GetActor());
 	}
 
-	return CachedAnimInstance;
+	return CachedAnimInstance.Get();
 }
 
 USkeletalMeshComponent* FContextualAnimSceneBinding::GetSkeletalMeshComponent() const
 {
-	if(CachedSkeletalMesh == nullptr)
+	if(!CachedSkeletalMesh.IsValid())
 	{
 		CachedSkeletalMesh = UContextualAnimUtilities::TryGetSkeletalMeshComponent(GetActor());
 	}
 
-	return CachedSkeletalMesh;
+	return CachedSkeletalMesh.Get();
 }
 
 FAnimMontageInstance* FContextualAnimSceneBinding::GetAnimMontageInstance() const
diff --git a/Engine/Plugins/Experimental/Animation/ContextualAnimation/Source/ContextualAnimation/Public/ContextualAnimTypes.h b/Engine/Plugins/Experimental/Animation/ContextualAnimation/Source/ContextualAnimation/Public/ContextualAnimTypes.h
--- a/Engine/Plugins/Experimental/Animation/ContextualAnimation/Source/ContextualAnimation/Public/ContextualAnimTypes.h
+++ b/Engine/Plugins/Experimental/Animation/ContextualAnimation/Source/ContextualAnimation/Public/ContextualAnimTypes.h
@@ -367,13 +367,20 @@ private:
 	friend UContextualAnimSceneInstance;
 	friend struct FContextualAnimSceneBindings;
 
+	UPROPERTY()
 	FContextualAnimSceneBindingContext Context;
 
+	UPROPERTY()
 	int32 AnimTrackIdx = INDEX_NONE;
 
-	mutable TObjectPtr<UContextualAnimSceneActorComponent> CachedSceneActorComp = nullptr;
-	mutable TObjectPtr<UAnimInstance> CachedAnimInstance = nullptr;
-	mutable TObjectPtr<USkeletalMeshComponent> CachedSkeletalMesh = nullptr;
+	UPROPERTY()
+	mutable TWeakObjectPtr<UContextualAnimSceneActorComponent> CachedSceneActorComp = nullptr;
+
+	UPROPERTY()
+	mutable TWeakObjectPtr<UAnimInstance> CachedAnimInstance = nullptr;
+
+	UPROPERTY()
+	mutable TWeakObjectPtr<USkeletalMeshComponent> CachedSkeletalMesh = nullptr;
 };
 
 template<>
