Added missing include to fix non unity build

diff --git a/Engine/Plugins/Experimental/Animation/ContextualAnimation/Source/ContextualAnimation/Private/ContextualAnimTypes.cpp b/Engine/Plugins/Experimental/Animation/ContextualAnimation/Source/ContextualAnimation/Private/ContextualAnimTypes.cpp
--- a/Engine/Plugins/Experimental/Animation/ContextualAnimation/Source/ContextualAnimation/Private/ContextualAnimTypes.cpp
+++ b/Engine/Plugins/Experimental/Animation/ContextualAnimation/Source/ContextualAnimation/Private/ContextualAnimTypes.cpp
@@ -10,6 +10,7 @@
 #include "ContextualAnimSelectionCriterion.h"
 #include "ContextualAnimSceneActorComponent.h"
 #include "ContextualAnimSceneAsset.h"
+#include "Components/SkeletalMeshComponent.h"
 
 #include UE_INLINE_GENERATED_CPP_BY_NAME(ContextualAnimTypes)
 
