Contextual Anim: Removed custom net serialize from bindings related structures and ptr to scene instance in the bindings

diff --git a/Engine/Plugins/Experimental/Animation/ContextualAnimation/Source/ContextualAnimation/Private/ContextualAnimManager.cpp b/Engine/Plugins/Experimental/Animation/ContextualAnimation/Source/ContextualAnimation/Private/ContextualAnimManager.cpp
--- a/Engine/Plugins/Experimental/Animation/ContextualAnimation/Source/ContextualAnimation/Private/ContextualAnimManager.cpp
+++ b/Engine/Plugins/Experimental/Animation/ContextualAnimation/Source/ContextualAnimation/Private/ContextualAnimManager.cpp
@@ -144,7 +144,6 @@ UContextualAnimSceneInstance* UContextualAnimManager::ForceStartScene(const UCon
 	const UClass* Class = SceneAsset.GetSceneInstanceClass();
 	UContextualAnimSceneInstance* NewInstance = Class ? NewObject<UContextualAnimSceneInstance>(this, Class) : NewObject<UContextualAnimSceneInstance>(this);
 	NewInstance->SceneAsset = &SceneAsset;
-	Bindings.SceneInstancePtr = NewInstance;
 
 	if (Params.Pivots.IsEmpty())
 	{
@@ -203,7 +202,6 @@ UContextualAnimSceneInstance* UContextualAnimManager::TryStartScene(const UConte
 		const UClass* Class = SceneAsset.GetSceneInstanceClass();
 		UContextualAnimSceneInstance* NewInstance = Class ? NewObject<UContextualAnimSceneInstance>(this, Class) : NewObject<UContextualAnimSceneInstance>(this);
 		NewInstance->SceneAsset = &SceneAsset;
-		Bindings.SceneInstancePtr = NewInstance;
 
 		if (Params.Pivots.IsEmpty())
 		{
diff --git a/Engine/Plugins/Experimental/Animation/ContextualAnimation/Source/ContextualAnimation/Private/ContextualAnimSceneActorComponent.cpp b/Engine/Plugins/Experimental/Animation/ContextualAnimation/Source/ContextualAnimation/Private/ContextualAnimSceneActorComponent.cpp
--- a/Engine/Plugins/Experimental/Animation/ContextualAnimation/Source/ContextualAnimation/Private/ContextualAnimSceneActorComponent.cpp
+++ b/Engine/Plugins/Experimental/Animation/ContextualAnimation/Source/ContextualAnimation/Private/ContextualAnimSceneActorComponent.cpp
@@ -160,17 +160,8 @@ bool UContextualAnimSceneActorComponent::ServerEarlyOutContextualAnimScene_Valid
 	return true;
 }
 
-void UContextualAnimSceneActorComponent::OnRep_Bindings(const FContextualAnimSceneBindings& LastRepBindings)
+void UContextualAnimSceneActorComponent::OnRep_Bindings()
 {
-	// @TODO: This need more investigation but for now it prevents an issue caused by this OnRep_ triggering even when there is no (obvious) change in the data
-	if(RepBindings.GetID() == LastRepBindings.GetID() && RepBindings.IsValid() && LastRepBindings.IsValid())
-	{
-		UE_LOG(LogContextualAnim, Warning, TEXT("%-21s UContextualAnimSceneActorComponent::OnRep_Bindings Actor: %s RepBindings Id: %d LastRepBindings Id: %d"),
-			*UEnum::GetValueAsString(TEXT("Engine.ENetRole"), GetOwner()->GetLocalRole()), *GetNameSafe(GetOwner()), RepBindings.GetID(), LastRepBindings.GetID());
-
-		return;
-	}
-
 	UE_LOG(LogContextualAnim, Verbose, TEXT("%-21s UContextualAnimSceneActorComponent::OnRep_Bindings Actor: %s RepBindings Id: %d Num: %d Bindings Id: %d Num: %d"),
 		*UEnum::GetValueAsString(TEXT("Engine.ENetRole"), GetOwner()->GetLocalRole()), *GetNameSafe(GetOwner()), RepBindings.GetID(), RepBindings.Num(), Bindings.GetID(), Bindings.Num());
 
diff --git a/Engine/Plugins/Experimental/Animation/ContextualAnimation/Source/ContextualAnimation/Private/ContextualAnimTypes.cpp b/Engine/Plugins/Experimental/Animation/ContextualAnimation/Source/ContextualAnimation/Private/ContextualAnimTypes.cpp
--- a/Engine/Plugins/Experimental/Animation/ContextualAnimation/Source/ContextualAnimation/Private/ContextualAnimTypes.cpp
+++ b/Engine/Plugins/Experimental/Animation/ContextualAnimation/Source/ContextualAnimation/Private/ContextualAnimTypes.cpp
@@ -263,14 +263,6 @@ FVector FContextualAnimSceneBindingContext::GetVelocity() const
 	return FVector::ZeroVector;
 }
 
-bool FContextualAnimSceneBindingContext::NetSerialize(FArchive& Ar, class UPackageMap* Map, bool& bOutSuccess)
-{
-	//@TODO: Serialize optional properties
-	Ar << Actor;
-	bOutSuccess = true;
-	return true;
-}
-
 // FContextualAnimSceneBinding
 ///////////////////////////////////////////////////////////////////////
 
@@ -351,14 +343,6 @@ int32 FContextualAnimSceneBinding::GetCurrentSectionIndex() const
 	return INDEX_NONE;
 }
 
-bool FContextualAnimSceneBinding::NetSerialize(FArchive& Ar, class UPackageMap* Map, bool& bOutSuccess)
-{
-	Context.NetSerialize(Ar, Map, bOutSuccess);
-	Ar << AnimTrackIdx;
-	bOutSuccess = true;
-	return true;
-}
-
 // FContextualAnimSceneBindings
 ///////////////////////////////////////////////////////////////////////
 
@@ -390,13 +374,11 @@ void FContextualAnimSceneBindings::Reset()
 	SectionIdx = INDEX_NONE;
 	AnimSetIdx = INDEX_NONE;
 	Data.Reset();
-	SceneInstancePtr.Reset();
 }
 
 void FContextualAnimSceneBindings::Clear()
 {
 	Data.Reset();
-	SceneInstancePtr.Reset();
 }
 
 const FContextualAnimSceneBinding* FContextualAnimSceneBindings::GetSyncLeader() const
@@ -465,16 +447,6 @@ const FName& FContextualAnimSceneBindings::GetRoleFromBinding(const FContextualA
 	return GetAnimTrackFromBinding(Binding).Role;
 }
 
-bool FContextualAnimSceneBindings::NetSerialize(FArchive& Ar, class UPackageMap* Map, bool& bOutSuccess)
-{
-	Ar << Id;
-	Ar << SceneAsset;
-	Ar << SectionIdx;
-	Ar << AnimSetIdx;
-	bOutSuccess = SafeNetSerializeTArray_WithNetSerialize<10>(Ar, Data, Map);
-	return true;
-}
-
 bool FContextualAnimSceneBindings::CheckConditions(const UContextualAnimSceneAsset& SceneAsset, int32 SectionIdx, int32 AnimSetIdx, const TMap<FName, FContextualAnimSceneBindingContext>& Params)
 {
 	check(SceneAsset.HasValidData());
diff --git a/Engine/Plugins/Experimental/Animation/ContextualAnimation/Source/ContextualAnimation/Public/ContextualAnimSceneActorComponent.h b/Engine/Plugins/Experimental/Animation/ContextualAnimation/Source/ContextualAnimation/Public/ContextualAnimSceneActorComponent.h
--- a/Engine/Plugins/Experimental/Animation/ContextualAnimation/Source/ContextualAnimation/Public/ContextualAnimSceneActorComponent.h
+++ b/Engine/Plugins/Experimental/Animation/ContextualAnimation/Source/ContextualAnimation/Public/ContextualAnimSceneActorComponent.h
@@ -105,7 +105,7 @@ protected:
 	void OnTickPose(class USkinnedMeshComponent* SkinnedMeshComponent, float DeltaTime, bool bNeedsValidRootMotion);
 
 	UFUNCTION()
-	void OnRep_Bindings(const FContextualAnimSceneBindings& LastRepBindings);
+	void OnRep_Bindings();
 
 	void SetIgnoreCollisionWithOtherActors(bool bValue) const;
 
diff --git a/Engine/Plugins/Experimental/Animation/ContextualAnimation/Source/ContextualAnimation/Public/ContextualAnimTypes.h b/Engine/Plugins/Experimental/Animation/ContextualAnimation/Source/ContextualAnimation/Public/ContextualAnimTypes.h
--- a/Engine/Plugins/Experimental/Animation/ContextualAnimation/Source/ContextualAnimation/Public/ContextualAnimTypes.h
+++ b/Engine/Plugins/Experimental/Animation/ContextualAnimation/Source/ContextualAnimation/Public/ContextualAnimTypes.h
@@ -16,7 +16,6 @@ class UAnimInstance;
 class UAnimSequenceBase;
 class UContextualAnimSelectionCriterion;
 class UContextualAnimSceneAsset;
-class UContextualAnimSceneInstance;
 class UContextualAnimSceneActorComponent;
 struct FAnimMontageInstance;
 
@@ -301,10 +300,9 @@ struct CONTEXTUALANIMATION_API FContextualAnimSceneBindingContext
 
 	//@TODO: Add accessors for GameplayTags
 
-	bool NetSerialize(FArchive& Ar, class UPackageMap* Map, bool& bOutSuccess);
-
 private:
 
+	UPROPERTY()
 	TWeakObjectPtr<AActor> Actor = nullptr;
 
 	TOptional<FTransform> ExternalTransform;
@@ -314,15 +312,6 @@ private:
 	FGameplayTagContainer ExternalGameplayTags;
 };
 
-template<>
-struct TStructOpsTypeTraits<FContextualAnimSceneBindingContext> : public TStructOpsTypeTraitsBase2<FContextualAnimSceneBindingContext>
-{
-	enum
-	{
-		WithNetSerializer = true
-	};
-};
-
 /** Represent an actor bound to a role in the scene */
 USTRUCT(BlueprintType)
 struct CONTEXTUALANIMATION_API FContextualAnimSceneBinding
@@ -358,13 +347,11 @@ struct CONTEXTUALANIMATION_API FContextualAnimSceneBinding
 
 	UContextualAnimSceneActorComponent* GetSceneActorComponent() const;
 
-	bool NetSerialize(FArchive& Ar, class UPackageMap* Map, bool& bOutSuccess);
-
 	static const FContextualAnimSceneBinding InvalidBinding;
 
 private:
 
-	friend UContextualAnimSceneInstance;
+	friend class UContextualAnimSceneInstance;
 	friend struct FContextualAnimSceneBindings;
 
 	UPROPERTY()
@@ -373,25 +360,16 @@ private:
 	UPROPERTY()
 	int32 AnimTrackIdx = INDEX_NONE;
 
-	UPROPERTY()
+	UPROPERTY(NotReplicated)
 	mutable TWeakObjectPtr<UContextualAnimSceneActorComponent> CachedSceneActorComp = nullptr;
 
-	UPROPERTY()
+	UPROPERTY(NotReplicated)
 	mutable TWeakObjectPtr<UAnimInstance> CachedAnimInstance = nullptr;
 
-	UPROPERTY()
+	UPROPERTY(NotReplicated)
 	mutable TWeakObjectPtr<USkeletalMeshComponent> CachedSkeletalMesh = nullptr;
 };
 
-template<>
-struct TStructOpsTypeTraits<FContextualAnimSceneBinding> : public TStructOpsTypeTraitsBase2<FContextualAnimSceneBinding>
-{
-	enum
-	{
-		WithNetSerializer = true
-	};
-};
-
 USTRUCT(BlueprintType)
 struct CONTEXTUALANIMATION_API FContextualAnimSceneBindings
 {
@@ -417,7 +395,6 @@ struct CONTEXTUALANIMATION_API FContextualAnimSceneBindings
 	FORCEINLINE int32 Num() const { return Data.Num(); }
 	FORCEINLINE int32 Add(const FContextualAnimSceneBinding& NewData) { return Data.Add(NewData); }
 	FORCEINLINE const TArray<FContextualAnimSceneBinding>& GetBindings() const { return Data; }
-	FORCEINLINE const UContextualAnimSceneInstance* GetSceneInstance() const { return SceneInstancePtr.Get(); }
 
 	FORCEINLINE TArray<FContextualAnimSceneBinding>::RangedForIteratorType      begin() { return Data.begin(); }
 	FORCEINLINE TArray<FContextualAnimSceneBinding>::RangedForConstIteratorType begin() const { return Data.begin(); }
@@ -449,8 +426,6 @@ struct CONTEXTUALANIMATION_API FContextualAnimSceneBindings
 
 	void Clear();
 
-	bool NetSerialize(FArchive& Ar, class UPackageMap* Map, bool& bOutSuccess);
-
 	void GenerateUniqueId();
 
 private:
@@ -472,19 +447,6 @@ private:
 	/** List of actors bound to each role in the SceneAsset */
 	UPROPERTY()
 	TArray<FContextualAnimSceneBinding> Data;
-
-	/** Ptr back to the scene instance we belong to (if any) */
-	UPROPERTY()
-	TWeakObjectPtr<const UContextualAnimSceneInstance> SceneInstancePtr = nullptr;
-};
-
-template<>
-struct TStructOpsTypeTraits<FContextualAnimSceneBindings> : public TStructOpsTypeTraitsBase2<FContextualAnimSceneBindings>
-{
-	enum
-	{
-		WithNetSerializer = true
-	};
 };
 
 USTRUCT(BlueprintType)
