Contextual Anim: Binding's GetTransform method now returns the transform of the scene actor comp instead of the actor transform to allow users to have explicit control over the transform used for alignment.

diff --git a/Engine/Plugins/Experimental/Animation/ContextualAnimation/Source/ContextualAnimation/Private/ContextualAnimTypes.cpp b/Engine/Plugins/Experimental/Animation/ContextualAnimation/Source/ContextualAnimation/Private/ContextualAnimTypes.cpp
--- a/Engine/Plugins/Experimental/Animation/ContextualAnimation/Source/ContextualAnimation/Private/ContextualAnimTypes.cpp
+++ b/Engine/Plugins/Experimental/Animation/ContextualAnimation/Source/ContextualAnimation/Private/ContextualAnimTypes.cpp
@@ -237,13 +237,16 @@ void FContextualAnimSceneBindingContext::SetExternalTransform(const FTransform&
 
 FTransform FContextualAnimSceneBindingContext::GetTransform() const
 {
+	// If created with an external transform, used that one to represent the location/rotation of the actor
 	if (ExternalTransform.IsSet())
 	{
 		return ExternalTransform.GetValue();
 	}
-	else if (AActor* ActorPtr = GetActor())
+	// If no external transform is provided, use the transform of the scene actor comp
+	// We use this instead of the actor transform so we can have explicit control over the transform that represents the location/rotation of this actor during alignment
+	else if (UContextualAnimSceneActorComponent* Comp = GetSceneActorComponent())
 	{
-		return ActorPtr->GetActorTransform();
+		return Comp->GetComponentTransform();
 	}
 
 	return FTransform::Identity;
@@ -263,25 +266,11 @@ FVector FContextualAnimSceneBindingContext::GetVelocity() const
 	return FVector::ZeroVector;
 }
 
-// FContextualAnimSceneBinding
-///////////////////////////////////////////////////////////////////////
-
-FContextualAnimSceneBinding::FContextualAnimSceneBinding(const FContextualAnimSceneBindingContext& InContext, const FContextualAnimTrack& InAnimTrack)
-{
-	Context = InContext;
-	AnimTrackIdx = InAnimTrack.AnimTrackIdx;
-}
-
-void FContextualAnimSceneBinding::SetAnimTrack(const FContextualAnimTrack& InAnimTrack) 
-{
-	AnimTrackIdx = InAnimTrack.AnimTrackIdx;
-}
-
-UContextualAnimSceneActorComponent* FContextualAnimSceneBinding::GetSceneActorComponent() const
+UContextualAnimSceneActorComponent* FContextualAnimSceneBindingContext::GetSceneActorComponent() const
 {
 	if (!CachedSceneActorComp.IsValid())
 	{
-		if (AActor* Actor = Context.GetActor())
+		if (Actor.IsValid())
 		{
 			CachedSceneActorComp = Actor->FindComponentByClass<UContextualAnimSceneActorComponent>();
 		}
@@ -290,7 +279,7 @@ UContextualAnimSceneActorComponent* FContextualAnimSceneBinding::GetSceneActorCo
 	return CachedSceneActorComp.Get();
 }
 
-UAnimInstance* FContextualAnimSceneBinding::GetAnimInstance() const
+UAnimInstance* FContextualAnimSceneBindingContext::GetAnimInstance() const
 {
 	if (!CachedAnimInstance.IsValid())
 	{
@@ -300,9 +289,9 @@ UAnimInstance* FContextualAnimSceneBinding::GetAnimInstance() const
 	return CachedAnimInstance.Get();
 }
 
-USkeletalMeshComponent* FContextualAnimSceneBinding::GetSkeletalMeshComponent() const
+USkeletalMeshComponent* FContextualAnimSceneBindingContext::GetSkeletalMeshComponent() const
 {
-	if(!CachedSkeletalMesh.IsValid())
+	if (!CachedSkeletalMesh.IsValid())
 	{
 		CachedSkeletalMesh = UContextualAnimUtilities::TryGetSkeletalMeshComponent(GetActor());
 	}
@@ -310,6 +299,20 @@ USkeletalMeshComponent* FContextualAnimSceneBinding::GetSkeletalMeshComponent()
 	return CachedSkeletalMesh.Get();
 }
 
+// FContextualAnimSceneBinding
+///////////////////////////////////////////////////////////////////////
+
+FContextualAnimSceneBinding::FContextualAnimSceneBinding(const FContextualAnimSceneBindingContext& InContext, const FContextualAnimTrack& InAnimTrack)
+{
+	Context = InContext;
+	AnimTrackIdx = InAnimTrack.AnimTrackIdx;
+}
+
+void FContextualAnimSceneBinding::SetAnimTrack(const FContextualAnimTrack& InAnimTrack) 
+{
+	AnimTrackIdx = InAnimTrack.AnimTrackIdx;
+}
+
 FAnimMontageInstance* FContextualAnimSceneBinding::GetAnimMontageInstance() const
 {
 	if (UAnimInstance* AnimInstance = GetAnimInstance())
diff --git a/Engine/Plugins/Experimental/Animation/ContextualAnimation/Source/ContextualAnimation/Public/ContextualAnimTypes.h b/Engine/Plugins/Experimental/Animation/ContextualAnimation/Source/ContextualAnimation/Public/ContextualAnimTypes.h
--- a/Engine/Plugins/Experimental/Animation/ContextualAnimation/Source/ContextualAnimation/Public/ContextualAnimTypes.h
+++ b/Engine/Plugins/Experimental/Animation/ContextualAnimation/Source/ContextualAnimation/Public/ContextualAnimTypes.h
@@ -292,6 +292,12 @@ struct CONTEXTUALANIMATION_API FContextualAnimSceneBindingContext
 
 	AActor* GetActor() const { return Actor.Get(); }
 
+	UAnimInstance* GetAnimInstance() const;
+
+	USkeletalMeshComponent* GetSkeletalMeshComponent() const;
+
+	UContextualAnimSceneActorComponent* GetSceneActorComponent() const;
+
 	void SetExternalTransform(const FTransform& InTransform);
 
 	FTransform GetTransform() const;
@@ -305,6 +311,15 @@ private:
 	UPROPERTY()
 	TWeakObjectPtr<AActor> Actor = nullptr;
 
+	UPROPERTY(NotReplicated)
+	mutable TWeakObjectPtr<UContextualAnimSceneActorComponent> CachedSceneActorComp = nullptr;
+
+	UPROPERTY(NotReplicated)
+	mutable TWeakObjectPtr<UAnimInstance> CachedAnimInstance = nullptr;
+
+	UPROPERTY(NotReplicated)
+	mutable TWeakObjectPtr<USkeletalMeshComponent> CachedSkeletalMesh = nullptr;
+
 	TOptional<FTransform> ExternalTransform;
 
 	TOptional<FVector> ExternalVelocity;
@@ -326,7 +341,9 @@ struct CONTEXTUALANIMATION_API FContextualAnimSceneBinding
 	FORCEINLINE AActor* GetActor() const { return GetContext().GetActor(); }
 	FORCEINLINE FTransform GetTransform() const { return GetContext().GetTransform(); }
 	FORCEINLINE FVector GetVelocity()  const { return GetContext().GetVelocity(); }
-
+	FORCEINLINE UAnimInstance* GetAnimInstance() const { return Context.GetAnimInstance(); }
+	FORCEINLINE USkeletalMeshComponent* GetSkeletalMeshComponent() const { return Context.GetSkeletalMeshComponent(); }
+	FORCEINLINE UContextualAnimSceneActorComponent* GetSceneActorComponent() const { return Context.GetSceneActorComponent(); }
 	FORCEINLINE int32 GetAnimTrackIdx() const { return AnimTrackIdx; }
 	
 	void SetAnimTrack(const FContextualAnimTrack& InAnimTrack);
@@ -341,12 +358,6 @@ struct CONTEXTUALANIMATION_API FContextualAnimSceneBinding
 	/** Returns the ActiveMontageInstance or null in the case of static actors */
 	FAnimMontageInstance* GetAnimMontageInstance() const;
 
-	UAnimInstance* GetAnimInstance() const;
-
-	USkeletalMeshComponent* GetSkeletalMeshComponent() const;
-
-	UContextualAnimSceneActorComponent* GetSceneActorComponent() const;
-
 	static const FContextualAnimSceneBinding InvalidBinding;
 
 private:
@@ -359,15 +370,6 @@ private:
 
 	UPROPERTY()
 	int32 AnimTrackIdx = INDEX_NONE;
-
-	UPROPERTY(NotReplicated)
-	mutable TWeakObjectPtr<UContextualAnimSceneActorComponent> CachedSceneActorComp = nullptr;
-
-	UPROPERTY(NotReplicated)
-	mutable TWeakObjectPtr<UAnimInstance> CachedAnimInstance = nullptr;
-
-	UPROPERTY(NotReplicated)
-	mutable TWeakObjectPtr<USkeletalMeshComponent> CachedSkeletalMesh = nullptr;
 };
 
 USTRUCT(BlueprintType)
