Contextual Anim Editor: Fixed bug making actors collide in the preview scene.

diff --git a/Engine/Plugins/Experimental/Animation/ContextualAnimation/Source/ContextualAnimationEditor/Private/ContextualAnimViewModel.cpp b/Engine/Plugins/Experimental/Animation/ContextualAnimation/Source/ContextualAnimationEditor/Private/ContextualAnimViewModel.cpp
--- a/Engine/Plugins/Experimental/Animation/ContextualAnimation/Source/ContextualAnimationEditor/Private/ContextualAnimViewModel.cpp
+++ b/Engine/Plugins/Experimental/Animation/ContextualAnimation/Source/ContextualAnimationEditor/Private/ContextualAnimViewModel.cpp
@@ -677,6 +677,15 @@ void FContextualAnimViewModel::StartSimulateMode()
 	MovieScene->SetReadOnly(true);
 
 	SimulateModeState = ESimulateModeState::Paused;
+
+	// Stop montage instance used during normal mode to preview the animations
+	for (const FContextualAnimSceneBinding& Binding : SceneBindings)
+	{
+		if (UAnimInstance* AnimInstance = Binding.GetAnimInstance())
+		{
+			AnimInstance->StopAllMontages(0);
+		}
+	}
 }
 
 void FContextualAnimViewModel::StopSimulateMode()
