Contextual Anim: Added methods to allow an actor to late join an interaction, transition all the actors in an interaction to a different section and transition a single actor to a different section.

diff --git a/Engine/Plugins/Experimental/Animation/ContextualAnimation/Source/ContextualAnimation/Private/ContextualAnimSceneActorComponent.cpp b/Engine/Plugins/Experimental/Animation/ContextualAnimation/Source/ContextualAnimation/Private/ContextualAnimSceneActorComponent.cpp
--- a/Engine/Plugins/Experimental/Animation/ContextualAnimation/Source/ContextualAnimation/Private/ContextualAnimSceneActorComponent.cpp
+++ b/Engine/Plugins/Experimental/Animation/ContextualAnimation/Source/ContextualAnimation/Private/ContextualAnimSceneActorComponent.cpp
@@ -3,6 +3,7 @@
 #include "ContextualAnimSceneActorComponent.h"
 #include "Animation/AnimInstance.h"
 #include "Animation/AnimMontage.h"
+#include "GameFramework/Character.h"
 #include "GameFramework/CharacterMovementComponent.h"
 #include "ContextualAnimSelectionCriterion.h"
 #include "Components/SkeletalMeshComponent.h"
@@ -18,6 +19,9 @@
 #include "PrimitiveSceneProxy.h"
 #include "PrimitiveViewRelevance.h"
 #include "SceneManagement.h"
+#include "MotionWarpingComponent.h"
+#include "AnimNotifyState_MotionWarping.h"
+#include "RootMotionModifier.h"
 
 #include UE_INLINE_GENERATED_CPP_BY_NAME(ContextualAnimSceneActorComponent)
 
@@ -26,6 +30,13 @@ TAutoConsoleVariable<int32> CVarContextualAnimIKDebug(TEXT("a.ContextualAnim.IK.
 TAutoConsoleVariable<float> CVarContextualAnimIKDrawDebugLifetime(TEXT("a.ContextualAnim.IK.DrawDebugLifetime"), 0, TEXT("Draw Debug Duration"));
 #endif
 
+void FContextualAnimRepData::IncrementRepCounter()
+{
+	static uint8 Counter = 0;
+	Counter = Counter < UINT8_MAX ? Counter + 1 : 0;
+	RepCounter = Counter;
+}
+
 UContextualAnimSceneActorComponent::UContextualAnimSceneActorComponent(const FObjectInitializer& ObjectInitializer)
 	: Super(ObjectInitializer)
 {
@@ -43,6 +54,9 @@ void UContextualAnimSceneActorComponent::GetLifetimeReplicatedProps(TArray< FLif
 	FDoRepLifetimeParams Params;
 	Params.bIsPushBased = true;
 	DOREPLIFETIME_WITH_PARAMS_FAST(UContextualAnimSceneActorComponent, RepBindings, Params);
+	DOREPLIFETIME_WITH_PARAMS_FAST(UContextualAnimSceneActorComponent, RepLateJoinData, Params);
+	DOREPLIFETIME_WITH_PARAMS_FAST(UContextualAnimSceneActorComponent, RepPlayAnimData, Params);
+	DOREPLIFETIME_WITH_PARAMS_FAST(UContextualAnimSceneActorComponent, RepTransitionData, Params);
 }
 
 bool UContextualAnimSceneActorComponent::IsOwnerLocallyControlled() const
@@ -55,6 +69,307 @@ bool UContextualAnimSceneActorComponent::IsOwnerLocallyControlled() const
 	return false;
 }
 
+void UContextualAnimSceneActorComponent::PlayAnimation_Internal(UAnimSequenceBase* Animation, float StartTime, bool bSyncPlaybackTime)
+{
+	TGuardValue<bool> UpdateGuard(bGuardAnimEvents, true);
+
+	if (UAnimInstance* AnimInstance = UContextualAnimUtilities::TryGetAnimInstance(GetOwner()))
+	{
+		UE_LOG(LogContextualAnim, VeryVerbose, TEXT("%-21s \t\tUContextualAnimSceneActorComponent::PlayAnimation_Internal Playing Animation. Actor: %s Anim: %s StartTime: %f bSyncPlaybackTime: %d"),
+			*UEnum::GetValueAsString(TEXT("Engine.ENetRole"), GetOwner()->GetLocalRole()), *GetNameSafe(GetOwner()), *GetNameSafe(Animation), StartTime, bSyncPlaybackTime);
+
+		//@TODO: Add support for dynamic montage
+		UAnimMontage* AnimMontage = Cast<UAnimMontage>(Animation);
+		AnimInstance->Montage_Play(AnimMontage, 1.f, EMontagePlayReturnType::MontageLength, StartTime);
+
+		AnimInstance->OnMontageBlendingOut.AddUniqueDynamic(this, &UContextualAnimSceneActorComponent::OnMontageBlendingOut);
+
+		if (bSyncPlaybackTime)
+		{
+			if (FAnimMontageInstance* MontageInstance = AnimInstance->GetActiveMontageInstance())
+			{
+				if (const FContextualAnimSceneBinding* SyncLeader = Bindings.GetSyncLeader())
+				{
+					if (SyncLeader->GetActor() != GetOwner())
+					{
+						FAnimMontageInstance* LeaderMontageInstance = SyncLeader->GetAnimMontageInstance();
+						if (LeaderMontageInstance && LeaderMontageInstance->Montage == Bindings.GetAnimTrackFromBinding(*SyncLeader).Animation && MontageInstance->GetMontageSyncLeader() == nullptr)
+						{
+							UE_LOG(LogContextualAnim, VeryVerbose, TEXT("%-21s \t\tUContextualAnimSceneActorComponent::PlayAnimation_Internal Syncing Animation. Actor: %s Anim: %s StartTime: %f bSyncPlaybackTime: %d"),
+								*UEnum::GetValueAsString(TEXT("Engine.ENetRole"), GetOwner()->GetLocalRole()), *GetNameSafe(GetOwner()), *GetNameSafe(Animation), StartTime, bSyncPlaybackTime);
+
+							MontageInstance->MontageSync_Follow(LeaderMontageInstance);
+						}
+					}
+				}
+			}
+		}
+	}
+
+	USkeletalMeshComponent* SkelMeshComp = UContextualAnimUtilities::TryGetSkeletalMeshComponent(GetOwner());
+	if (SkelMeshComp && !SkelMeshComp->OnTickPose.IsBoundToObject(this))
+	{
+		SkelMeshComp->OnTickPose.AddUObject(this, &UContextualAnimSceneActorComponent::OnTickPose);
+	}
+}
+
+void UContextualAnimSceneActorComponent::AddOrUpdateWarpTargets(int32 SectionIdx, int32 AnimSetIdx)
+{
+	if (const FContextualAnimSceneBinding* Binding = Bindings.FindBindingByActor(GetOwner()))
+	{
+		const UContextualAnimSceneAsset* Asset = Bindings.GetSceneAsset();
+		check(Asset);
+
+		const TArray<FContextualAnimSetPivotDefinition>& PivotDefs = Asset->GetAnimSetPivotDefinitionsInSection(SectionIdx);
+		if (PivotDefs.Num() > 0)
+		{
+			const FContextualAnimTrack* AnimTrack = Asset->GetAnimTrack(SectionIdx, AnimSetIdx, Bindings.GetRoleFromBinding(*Binding));
+			check(AnimTrack);
+
+			for (const FContextualAnimSetPivotDefinition& PivotDef : PivotDefs)
+			{
+				FContextualAnimSetPivot ScenePivot;
+				if (Bindings.CalculateAnimSetPivot(PivotDef, ScenePivot))
+				{
+					if (UMotionWarpingComponent* MotionWarpComp = GetOwner()->FindComponentByClass<UMotionWarpingComponent>())
+					{
+						const float Time = AnimTrack->GetSyncTimeForWarpSection(PivotDef.Name);
+						const FTransform TransformRelativeToScenePivot = Asset->GetAlignmentTransform(SectionIdx, AnimSetIdx, AnimTrack->AnimTrackIdx, PivotDef.Name, Time);
+						const FTransform WarpTarget = (TransformRelativeToScenePivot * ScenePivot.Transform);
+						MotionWarpComp->AddOrUpdateWarpTargetFromTransform(PivotDef.Name, WarpTarget);
+					}
+				}
+			}
+		}
+	}
+}
+
+bool UContextualAnimSceneActorComponent::LateJoinContextualAnimScene(AActor* Actor, FName Role)
+{
+	if (!Bindings.IsValid())
+	{
+		UE_LOG(LogContextualAnim, Warning, TEXT("%-21s UContextualAnimSceneActorComponent::LateJoinContextualAnimScene Invalid Bindings"), *UEnum::GetValueAsString(TEXT("Engine.ENetRole"), GetOwner()->GetLocalRole()));
+		return false;
+	}
+
+	// Redirect the request to the leader if needed. Technically this is not necessary but the idea here is that the leader of the interaction handles all the events for that interaction
+	// E.g the leader tells other actors to play the animation.
+	if (const FContextualAnimSceneBinding* Leader = Bindings.GetSyncLeader())
+	{
+		if (Leader->GetActor() != GetOwner())
+		{
+			if (UContextualAnimSceneActorComponent* Comp = Leader->GetSceneActorComponent())
+			{
+				return Comp->LateJoinContextualAnimScene(Actor, Role);
+			}
+		}
+	}
+
+	UE_LOG(LogContextualAnim, Verbose, TEXT("%-21s UContextualAnimSceneActorComponent::LateJoinContextualAnimScene Owner: %s Bindings Id: %d Section: %d Asset: %s. Requester: %s Role: %s"),
+		*UEnum::GetValueAsString(TEXT("Engine.ENetRole"), GetOwner()->GetLocalRole()), *GetNameSafe(GetOwner()), Bindings.GetID(), Bindings.GetSectionIdx(), *GetNameSafe(Bindings.GetSceneAsset()), *GetNameSafe(Actor), *Role.ToString());
+
+	// Play animation and set state on this new actor that is joining us and update bindings for everyone else
+	if (HandleLateJoin(Actor, Role))
+	{
+		// Replicate late join event. See OnRep_LateJoinData
+		if (GetOwner()->HasAuthority())
+		{
+			RepLateJoinData.Actor = Actor;
+			RepLateJoinData.Role = Role;
+			RepLateJoinData.IncrementRepCounter();
+			MARK_PROPERTY_DIRTY_FROM_NAME(UContextualAnimSceneActorComponent, RepLateJoinData, this);
+			GetOwner()->ForceNetUpdate();
+		}
+
+		return true;
+	}
+
+	return false;
+}
+
+bool UContextualAnimSceneActorComponent::HandleLateJoin(AActor* Actor, FName Role)
+{
+	if (!Bindings.BindActorToRole(*Actor, Role))
+	{
+		UE_LOG(LogContextualAnim, Warning, TEXT("%-21s UContextualAnimSceneActorComponent::HandleLateJoin Failed. Reason: Adding %s to the bindings for role: %s failed!"),
+			*UEnum::GetValueAsString(TEXT("Engine.ENetRole"), GetOwner()->GetLocalRole()), *GetNameSafe(Actor), *Role.ToString());
+
+		return false;
+	}
+
+	// Update the bindings on all the other actors too
+	for (const FContextualAnimSceneBinding& OtherBinding : Bindings)
+	{
+		if (OtherBinding.GetActor() != GetOwner() && OtherBinding.GetActor() != Actor)
+		{
+			if (UContextualAnimSceneActorComponent* Comp = OtherBinding.GetSceneActorComponent())
+			{
+				Comp->Bindings.BindActorToRole(*Actor, Role);
+			}
+		}
+	}
+
+	// Play animation and set state on this new actor that is joining us
+	if (const FContextualAnimSceneBinding* Binding = Bindings.FindBindingByActor(Actor))
+	{
+		if (UContextualAnimSceneActorComponent* Comp = Binding->GetSceneActorComponent())
+		{
+			Comp->LateJoinScene(Bindings);
+		}
+	}
+	
+
+	return true;
+}
+
+void UContextualAnimSceneActorComponent::LateJoinScene(const FContextualAnimSceneBindings& InBindings)
+{
+	if (Bindings.IsValid())
+	{
+		UE_LOG(LogContextualAnim, Verbose, TEXT("%-21s UContextualAnimSceneActorComponent::LateJoinScene Actor: %s Bindings Id: %d Section: %d Asset: %s. Leaving current scene"),
+			*UEnum::GetValueAsString(TEXT("Engine.ENetRole"), GetOwner()->GetLocalRole()), *GetNameSafe(GetOwner()), Bindings.GetID(), Bindings.GetSectionIdx(), *GetNameSafe(Bindings.GetSceneAsset()));
+
+		LeaveScene();
+	}
+
+	if (const FContextualAnimSceneBinding* Binding = InBindings.FindBindingByActor(GetOwner()))
+	{
+		UE_LOG(LogContextualAnim, Verbose, TEXT("%-21s UContextualAnimSceneActorComponent::LateJoinScene Actor: %s Role: %s Bindings Id: %d Section: %d Asset: %s"),
+			*UEnum::GetValueAsString(TEXT("Engine.ENetRole"), GetOwner()->GetLocalRole()), *GetNameSafe(GetOwner()), *InBindings.GetRoleFromBinding(*Binding).ToString(), InBindings.GetID(), InBindings.GetSectionIdx(), *GetNameSafe(InBindings.GetSceneAsset()));
+
+		Bindings = InBindings;
+
+		// For now when late joining an scene always play animation from first section
+		const int32 SectionIdx = 0;
+		const int32 AnimSetIdx = 0;
+		const FContextualAnimTrack* AnimTrack = Bindings.GetSceneAsset()->GetAnimTrack(SectionIdx, AnimSetIdx, Bindings.GetRoleFromBinding(*Binding));
+		check(AnimTrack);
+
+		PlayAnimation_Internal(AnimTrack->Animation, 0.f, false);
+
+		AddOrUpdateWarpTargets(SectionIdx, AnimSetIdx);
+
+		SetIgnoreCollisionWithOtherActors(true);
+
+		SetMovementState(AnimTrack->bRequireFlyingMode);
+	}
+}
+
+void UContextualAnimSceneActorComponent::OnRep_LateJoinData()
+{
+	// This is received by the leader of the interaction on every remote client
+
+	if (!Bindings.IsValid())
+	{
+		UE_LOG(LogContextualAnim, Warning, TEXT("%-21s UContextualAnimSceneActorComponent::OnRep_LateJoinData Invalid Bindings"), *UEnum::GetValueAsString(TEXT("Engine.ENetRole"), GetOwner()->GetLocalRole()));
+		return;
+	}
+
+	UE_LOG(LogContextualAnim, Verbose, TEXT("%-21s UContextualAnimSceneActorComponent::OnRep_LateJoinData Owner: %s Bindings Id: %d Section: %d Asset: %s. Requester: %s Role: %s"),
+		*UEnum::GetValueAsString(TEXT("Engine.ENetRole"), GetOwner()->GetLocalRole()), *GetNameSafe(GetOwner()), Bindings.GetID(), Bindings.GetSectionIdx(), *GetNameSafe(Bindings.GetSceneAsset()), *GetNameSafe(RepLateJoinData.Actor), *RepLateJoinData.Role.ToString());
+
+	// Play animation and set state on this new actor that is joining us and update bindings for everyone else
+	HandleLateJoin(RepLateJoinData.Actor, RepLateJoinData.Role);
+}
+
+bool UContextualAnimSceneActorComponent::TransitionContextualAnimScene(FName SectionName)
+{
+	if (!GetOwner()->HasAuthority())
+	{
+		return false;
+	}
+
+	// Redirect the request to the leader if needed. Technically this is not necessary but the idea here is that the leader of the interaction handles all the events for that interaction
+	// E.g the leader tells other actors to play the animation.
+	if (const FContextualAnimSceneBinding* Leader = Bindings.GetSyncLeader())
+	{
+		if (Leader->GetActor() != GetOwner())
+		{
+			if (UContextualAnimSceneActorComponent* Comp = Leader->GetSceneActorComponent())
+			{
+				return Comp->TransitionContextualAnimScene(SectionName);
+			}
+		}
+	}
+
+	if (const FContextualAnimSceneBinding* OwnerBinding = Bindings.FindBindingByActor(GetOwner()))
+	{
+		const int32 SectionIdx = Bindings.GetSceneAsset()->GetSectionIndex(SectionName);
+		if (SectionIdx != INDEX_NONE)
+		{
+			UE_LOG(LogContextualAnim, Log, TEXT("%-21s UContextualAnimSceneActorComponent::TransitionTo Actor: %s SectionName: %s"),
+				*UEnum::GetValueAsString(TEXT("Engine.ENetRole"), GetOwner()->GetLocalRole()), *GetNameSafe(GetOwner()), *SectionName.ToString());
+
+			HandleTransitionEveryone(SectionIdx, 0);
+
+			RepTransitionData.SectionIdx = SectionIdx;
+			RepTransitionData.AnimSetIdx = 0;
+			RepTransitionData.IncrementRepCounter();
+			MARK_PROPERTY_DIRTY_FROM_NAME(UContextualAnimSceneActorComponent, RepTransitionData, this);
+			GetOwner()->ForceNetUpdate();
+
+			return true;
+		}
+	}
+
+	return false;
+}
+
+void UContextualAnimSceneActorComponent::HandleTransitionEveryone(int32 NewSectionIdx, int32 NewAnimSetIdx)
+{
+	// Update Bindings internal data and play new animation for the leader first
+	// Note that for now we always transition to the first set in the section. We could run selection criteria here too but keeping it simple for now
+	HandleTransitionSelf(NewSectionIdx, NewAnimSetIdx);
+
+	// And now the same for everyone else
+	for (const FContextualAnimSceneBinding& Binding : Bindings)
+	{
+		if (Binding.GetActor() != GetOwner())
+		{
+			if(UContextualAnimSceneActorComponent* Comp = Binding.GetSceneActorComponent())
+			{
+				Comp->HandleTransitionSelf(NewSectionIdx, NewAnimSetIdx);
+			}
+		}
+	}
+}
+
+void UContextualAnimSceneActorComponent::HandleTransitionSelf(int32 NewSectionIdx, int32 NewAnimSetIdx)
+{
+	// Update bindings internal data so it points to the new section and new anim set
+	Bindings.TransitionTo(NewSectionIdx, NewAnimSetIdx);
+
+	// Play animation
+	//@TODO: Add support for dynamic montage
+	const FContextualAnimTrack& AnimTrack = Bindings.GetAnimTrackFromBinding(*Bindings.FindBindingByActor(GetOwner()));
+	PlayAnimation_Internal(AnimTrack.Animation, 0.f, true);
+
+	AddOrUpdateWarpTargets(NewSectionIdx, NewAnimSetIdx);
+}
+
+void UContextualAnimSceneActorComponent::PlayExternalAnimation(UAnimSequenceBase* Animation)
+{
+	if (!GetOwner()->HasAuthority())
+	{
+		return;
+	}
+
+	PlayAnimation_Internal(Animation, 0.f, false);
+
+	RepPlayAnimData.Animation = Animation;
+	RepPlayAnimData.IncrementRepCounter();
+	MARK_PROPERTY_DIRTY_FROM_NAME(UContextualAnimSceneActorComponent, RepPlayAnimData, this);
+	GetOwner()->ForceNetUpdate();
+}
+
+void UContextualAnimSceneActorComponent::OnRep_PlayAnimData()
+{
+	UE_LOG(LogContextualAnim, Verbose, TEXT("%-21s UContextualAnimSceneActorComponent::OnRep_PlayAnimData Owner: %s AnimToPlay: %s"),
+		*UEnum::GetValueAsString(TEXT("Engine.ENetRole"), GetOwner()->GetLocalRole()), *GetNameSafe(GetOwner()), *GetNameSafe(RepPlayAnimData.Animation));
+
+	PlayAnimation_Internal(RepPlayAnimData.Animation, 0.f, false);
+}
+
 bool UContextualAnimSceneActorComponent::StartContextualAnimScene(const FContextualAnimSceneBindings& InBindings)
 {
 	UE_LOG(LogContextualAnim, Log, TEXT("%-21s UContextualAnimSceneActorComponent::StartContextualAnim Actor: %s"),
@@ -78,9 +393,6 @@ bool UContextualAnimSceneActorComponent::StartContextualAnimScene(const FContext
 				}
 			}
 
-			//@TODO: Temp until we move the scene pivots to the bindings
-			UContextualAnimUtilities::BP_SceneBindings_AddOrUpdateWarpTargetsForBindings(InBindings);
-
 			RepBindings = InBindings;
 			MARK_PROPERTY_DIRTY_FROM_NAME(UContextualAnimSceneActorComponent, RepBindings, this);
 			GetOwner()->ForceNetUpdate();
@@ -91,9 +403,6 @@ bool UContextualAnimSceneActorComponent::StartContextualAnimScene(const FContext
 		{
 			JoinScene(InBindings);
 
-			//@TODO: Temp until we move the scene pivots to the bindings
-			UContextualAnimUtilities::BP_SceneBindings_AddOrUpdateWarpTargetsForBindings(InBindings);
-
 			ServerStartContextualAnimScene(InBindings);
 
 			return true;
@@ -160,6 +469,14 @@ bool UContextualAnimSceneActorComponent::ServerEarlyOutContextualAnimScene_Valid
 	return true;
 }
 
+void UContextualAnimSceneActorComponent::OnRep_TransitionData()
+{
+	UE_LOG(LogContextualAnim, Verbose, TEXT("%-21s UContextualAnimSceneActorComponent::OnRep_TransitionData Actor: %s"),
+		*UEnum::GetValueAsString(TEXT("Engine.ENetRole"), GetOwner()->GetLocalRole()), *GetNameSafe(GetOwner()));
+
+	HandleTransitionEveryone(RepTransitionData.SectionIdx, RepTransitionData.AnimSetIdx);
+}
+
 void UContextualAnimSceneActorComponent::OnRep_Bindings()
 {
 	UE_LOG(LogContextualAnim, Verbose, TEXT("%-21s UContextualAnimSceneActorComponent::OnRep_Bindings Actor: %s RepBindings Id: %d Num: %d Bindings Id: %d Num: %d"),
@@ -340,47 +657,14 @@ void UContextualAnimSceneActorComponent::JoinScene(const FContextualAnimSceneBin
 		Bindings = InBindings;
 
 		const FContextualAnimTrack& AnimTrack = Bindings.GetAnimTrackFromBinding(*Binding);
+		PlayAnimation_Internal(AnimTrack.Animation, 0.f, true);
 
-		if (UAnimInstance* AnimInstance = Binding->GetAnimInstance())
-		{
-			//@TODO: Add support for dynamic montage
-			UAnimMontage* AnimMontage = Cast<UAnimMontage>(AnimTrack.Animation);
-			AnimInstance->Montage_Play(AnimMontage, 1.f);
-
-			AnimInstance->OnMontageBlendingOut.AddUniqueDynamic(this, &UContextualAnimSceneActorComponent::OnMontageBlendingOut);
-		}
-
-		USkeletalMeshComponent* SkelMeshComp = Binding->GetSkeletalMeshComponent();
-		if (SkelMeshComp && !SkelMeshComp->OnTickPose.IsBoundToObject(this))
-		{
-			SkelMeshComp->OnTickPose.AddUObject(this, &UContextualAnimSceneActorComponent::OnTickPose);
-		}
+		AddOrUpdateWarpTargets(AnimTrack.SectionIdx, AnimTrack.AnimSetIdx);
 
 		// Disable collision between actors so they can align perfectly
 		SetIgnoreCollisionWithOtherActors(true);
 
-		if (UCharacterMovementComponent* MovementComp = GetOwner()->FindComponentByClass<UCharacterMovementComponent>())
-		{
-			// Save movement state before the interaction starts so we can restore it when it ends
-			CharacterPropertiesBackup.bIgnoreClientMovementErrorChecksAndCorrection = MovementComp->bIgnoreClientMovementErrorChecksAndCorrection;
-			CharacterPropertiesBackup.bAllowPhysicsRotationDuringAnimRootMotion = MovementComp->bAllowPhysicsRotationDuringAnimRootMotion;
-			CharacterPropertiesBackup.bUseControllerDesiredRotation = MovementComp->bUseControllerDesiredRotation;
-			CharacterPropertiesBackup.bOrientRotationToMovement = MovementComp->bOrientRotationToMovement;
-
-			// Disable movement correction.
-			MovementComp->bIgnoreClientMovementErrorChecksAndCorrection = true;
-
-			// Prevent physics rotation. During the interaction we want to be fully root motion driven
-			MovementComp->bAllowPhysicsRotationDuringAnimRootMotion = false;
-			MovementComp->bUseControllerDesiredRotation = false;
-			MovementComp->bOrientRotationToMovement = false;
-
-			//@TODO: Temp solution that assumes these interactions are not locally predicted and that is ok to be in flying mode during the entire animation
-			if (AnimTrack.bRequireFlyingMode && MovementComp->MovementMode != MOVE_Flying)
-			{
-				MovementComp->SetMovementMode(MOVE_Flying);
-			}
-		}
+		SetMovementState(AnimTrack.bRequireFlyingMode);
 
 		OnJoinedSceneDelegate.Broadcast(this);
 	}
@@ -421,20 +705,7 @@ void UContextualAnimSceneActorComponent::LeaveScene()
 		// We might want to add a more robust mechanism to avoid overriding a request to disable collision that may have been set by another system
 		SetIgnoreCollisionWithOtherActors(false);
 
-		if (UCharacterMovementComponent* MovementComp = GetOwner()->FindComponentByClass<UCharacterMovementComponent>())
-		{
-			// Restore movement state
-			MovementComp->bIgnoreClientMovementErrorChecksAndCorrection = CharacterPropertiesBackup.bIgnoreClientMovementErrorChecksAndCorrection;
-			MovementComp->bAllowPhysicsRotationDuringAnimRootMotion = CharacterPropertiesBackup.bAllowPhysicsRotationDuringAnimRootMotion;
-			MovementComp->bUseControllerDesiredRotation = CharacterPropertiesBackup.bUseControllerDesiredRotation;
-			MovementComp->bOrientRotationToMovement = CharacterPropertiesBackup.bOrientRotationToMovement;
-
-			//@TODO: Temp solution that assumes these interactions are not locally predicted and that is ok to be in flying mode during the entire animation
-			if (AnimTrack.bRequireFlyingMode && MovementComp->MovementMode == MOVE_Flying)
-			{
-				MovementComp->SetMovementMode(MOVE_Walking);
-			}
-		}
+		RestoreMovementState(AnimTrack.bRequireFlyingMode);
 
 		OnLeftSceneDelegate.Broadcast(this);
 
@@ -442,27 +713,73 @@ void UContextualAnimSceneActorComponent::LeaveScene()
 	}
 }
 
+void UContextualAnimSceneActorComponent::SetMovementState(bool bRequireFlyingMode)
+{
+	if (UCharacterMovementComponent* MovementComp = GetOwner()->FindComponentByClass<UCharacterMovementComponent>())
+	{
+		// Save movement state before the interaction starts so we can restore it when it ends
+		CharacterPropertiesBackup.bIgnoreClientMovementErrorChecksAndCorrection = MovementComp->bIgnoreClientMovementErrorChecksAndCorrection;
+		CharacterPropertiesBackup.bAllowPhysicsRotationDuringAnimRootMotion = MovementComp->bAllowPhysicsRotationDuringAnimRootMotion;
+		CharacterPropertiesBackup.bUseControllerDesiredRotation = MovementComp->bUseControllerDesiredRotation;
+		CharacterPropertiesBackup.bOrientRotationToMovement = MovementComp->bOrientRotationToMovement;
+
+		// Disable movement correction.
+		MovementComp->bIgnoreClientMovementErrorChecksAndCorrection = true;
+
+		// Prevent physics rotation. During the interaction we want to be fully root motion driven
+		MovementComp->bAllowPhysicsRotationDuringAnimRootMotion = false;
+		MovementComp->bUseControllerDesiredRotation = false;
+		MovementComp->bOrientRotationToMovement = false;
+
+		//@TODO: Temp solution that assumes these interactions are not locally predicted and that is ok to be in flying mode during the entire animation
+		if (bRequireFlyingMode && MovementComp->MovementMode != MOVE_Flying)
+		{
+			MovementComp->SetMovementMode(MOVE_Flying);
+		}
+	}
+}
+
+void UContextualAnimSceneActorComponent::RestoreMovementState(bool bRequireFlyingMode)
+{
+	if (UCharacterMovementComponent* MovementComp = GetOwner()->FindComponentByClass<UCharacterMovementComponent>())
+	{
+		// Restore movement state
+		MovementComp->bIgnoreClientMovementErrorChecksAndCorrection = CharacterPropertiesBackup.bIgnoreClientMovementErrorChecksAndCorrection;
+		MovementComp->bAllowPhysicsRotationDuringAnimRootMotion = CharacterPropertiesBackup.bAllowPhysicsRotationDuringAnimRootMotion;
+		MovementComp->bUseControllerDesiredRotation = CharacterPropertiesBackup.bUseControllerDesiredRotation;
+		MovementComp->bOrientRotationToMovement = CharacterPropertiesBackup.bOrientRotationToMovement;
+
+		//@TODO: Temp solution that assumes these interactions are not locally predicted and that is ok to be in flying mode during the entire animation
+		if (bRequireFlyingMode && MovementComp->MovementMode == MOVE_Flying)
+		{
+			MovementComp->SetMovementMode(MOVE_Walking);
+		}
+	}
+}
+
 void UContextualAnimSceneActorComponent::OnMontageBlendingOut(UAnimMontage* Montage, bool bInterrupted)
 {
+	if (bGuardAnimEvents)
+	{
+		return;
+	}
+
 	UE_LOG(LogContextualAnim, Verbose, TEXT("%-21s UContextualAnimSceneActorComponent::OnMontageBlendingOut Actor: %s Montage: %s bInterrupted: %d"),
 		*UEnum::GetValueAsString(TEXT("Engine.ENetRole"), GetOwner()->GetLocalRole()), *GetNameSafe(GetOwner()), *GetNameSafe(Montage), bInterrupted);
 
 	if (const FContextualAnimSceneBinding* Binding = Bindings.FindBindingByActor(GetOwner()))
 	{
-		if (Bindings.GetAnimTrackFromBinding(*Binding).Animation == Montage)
-		{
-			LeaveScene();
+		LeaveScene();
 
-			if (GetOwner()->HasAuthority())
+		if (GetOwner()->HasAuthority())
+		{
+			// Rep empty bindings if we were the initiator of this interaction.
+			if (RepBindings.IsValid())
 			{
-				// Rep empty bindings if we were the initiator of this interaction.
-				if(RepBindings.IsValid())
-				{
-					RepBindings.Clear();
-					MARK_PROPERTY_DIRTY_FROM_NAME(UContextualAnimSceneActorComponent, RepBindings, this);
+				RepBindings.Clear();
+				MARK_PROPERTY_DIRTY_FROM_NAME(UContextualAnimSceneActorComponent, RepBindings, this);
 
-					GetOwner()->ForceNetUpdate();
-				}
+				GetOwner()->ForceNetUpdate();
 			}
 		}
 	}
@@ -470,31 +787,8 @@ void UContextualAnimSceneActorComponent::OnMontageBlendingOut(UAnimMontage* Mont
 
 void UContextualAnimSceneActorComponent::OnTickPose(class USkinnedMeshComponent* SkinnedMeshComponent, float DeltaTime, bool bNeedsValidRootMotion)
 {
-	if (const FContextualAnimSceneBinding* Binding = Bindings.FindBindingByActor(GetOwner()))
-	{
-		// Synchronize playback time with the leader
-		FAnimMontageInstance* MontageInstance = Binding->GetAnimMontageInstance();
-		if (MontageInstance && MontageInstance->GetMontageSyncLeader() == nullptr)
-		{
-			if (const FContextualAnimSceneBinding* SyncLeader = Bindings.GetSyncLeader())
-			{
-				if (SyncLeader->GetActor() != GetOwner())
-				{
-					if (FAnimMontageInstance* LeaderMontageInstance = SyncLeader->GetAnimMontageInstance())
-					{
-						if (LeaderMontageInstance->Montage == Bindings.GetAnimTrackFromBinding(*SyncLeader).Animation &&
-							MontageInstance->Montage == Bindings.GetAnimTrackFromBinding(*Binding).Animation)
-						{
-							MontageInstance->MontageSync_Follow(LeaderMontageInstance);
-						}
-					}
-				}
-			}
-		}
-
-		//@TODO: Check for LOD to prevent this update if the actor is too far away
-		UpdateIKTargets();
-	}
+	//@TODO: Check for LOD to prevent this update if the actor is too far away
+	UpdateIKTargets();
 }
 
 void UContextualAnimSceneActorComponent::UpdateIKTargets()
diff --git a/Engine/Plugins/Experimental/Animation/ContextualAnimation/Source/ContextualAnimation/Private/ContextualAnimTypes.cpp b/Engine/Plugins/Experimental/Animation/ContextualAnimation/Source/ContextualAnimation/Private/ContextualAnimTypes.cpp
--- a/Engine/Plugins/Experimental/Animation/ContextualAnimation/Source/ContextualAnimation/Private/ContextualAnimTypes.cpp
+++ b/Engine/Plugins/Experimental/Animation/ContextualAnimation/Source/ContextualAnimation/Private/ContextualAnimTypes.cpp
@@ -387,7 +387,12 @@ void FContextualAnimSceneBindings::Clear()
 const FContextualAnimSceneBinding* FContextualAnimSceneBindings::GetSyncLeader() const
 {
 	//@TODO: Return first secondary binding as sync leader for now. This may have to be explicitly defined, either in the SceneAsset or when creating the bindings.
-	return Data.FindByPredicate([this](const FContextualAnimSceneBinding& Item) { return GetRoleFromBinding(Item) != SceneAsset->GetPrimaryRole(); });
+	return IsValid() ? Data.FindByPredicate([this](const FContextualAnimSceneBinding& Item) { return GetRoleFromBinding(Item) != SceneAsset->GetPrimaryRole(); }) : nullptr;
+}
+
+const FContextualAnimSceneBinding* FContextualAnimSceneBindings::GetPrimaryBinding() const
+{
+	return IsValid() ? Data.FindByPredicate([this](const FContextualAnimSceneBinding& Item) { return GetRoleFromBinding(Item) == SceneAsset->GetPrimaryRole(); }) : nullptr;
 }
 
 bool FContextualAnimSceneBindings::BindActorToRole(AActor& ActorRef, FName Role)
@@ -696,3 +701,22 @@ bool FContextualAnimSceneBindings::CalculateAnimSetPivot(const FContextualAnimSe
 
 	return false;
 }
+
+void FContextualAnimSceneBindings::TransitionTo(int32 NewSectionIdx, int32 NewAnimSetIdx)
+{
+	check(IsValid());
+
+	if (const FContextualAnimSet* AnimSet = SceneAsset->GetAnimSet(NewSectionIdx, NewAnimSetIdx))
+	{
+		SectionIdx = NewSectionIdx;
+		AnimSetIdx = NewAnimSetIdx;
+
+		for (FContextualAnimSceneBinding& Binding : Data)
+		{
+			if (const FContextualAnimTrack* AnimTrack = SceneAsset->GetAnimTrack(SectionIdx, AnimSetIdx, GetRoleFromBinding(Binding)))
+			{
+				Binding.AnimTrackIdx = AnimTrack->AnimTrackIdx;
+			}
+		}
+	}
+}
diff --git a/Engine/Plugins/Experimental/Animation/ContextualAnimation/Source/ContextualAnimation/Public/ContextualAnimSceneActorComponent.h b/Engine/Plugins/Experimental/Animation/ContextualAnimation/Source/ContextualAnimation/Public/ContextualAnimSceneActorComponent.h
--- a/Engine/Plugins/Experimental/Animation/ContextualAnimation/Source/ContextualAnimation/Public/ContextualAnimSceneActorComponent.h
+++ b/Engine/Plugins/Experimental/Animation/ContextualAnimation/Source/ContextualAnimation/Public/ContextualAnimSceneActorComponent.h
@@ -11,9 +11,61 @@ class AActor;
 class FPrimitiveSceneProxy;
 class UAnimInstance;
 class UAnimMontage;
+class UAnimSequenceBase;
 
 DECLARE_DYNAMIC_MULTICAST_DELEGATE_OneParam(FContextualAnimSceneActorCompDelegate, class UContextualAnimSceneActorComponent*, SceneActorComponent);
 
+/** Base struct for replicated data with a rep counter */
+USTRUCT()
+struct FContextualAnimRepData
+{
+	GENERATED_BODY()
+
+	/** Auto increment counter to ensure replication even if the data is the same */
+	UPROPERTY()
+	uint8 RepCounter = 0;
+
+	void IncrementRepCounter();
+};
+
+/** Used to replicate a late join event */
+USTRUCT()
+struct FContextualAnimRepLateJoinData : public FContextualAnimRepData
+{
+	GENERATED_BODY()
+
+	/** Actor that is joining the interaction */
+	UPROPERTY()
+	TObjectPtr<AActor> Actor = nullptr;
+
+	/** Role in the interaction the actor is gonna play */
+	UPROPERTY()
+	FName Role = NAME_None;
+};
+
+/** Used to replicate the playback of a 'random' animation */
+USTRUCT()
+struct FContextualAnimRepPlayAnimData : public FContextualAnimRepData
+{
+	GENERATED_BODY()
+
+	UPROPERTY()
+	TObjectPtr<UAnimSequenceBase> Animation = nullptr;
+};
+
+/** Used to replicate a transition to a new AnimSet in the interaction */
+USTRUCT()
+struct FContextualAnimRepTransitionData : public FContextualAnimRepData
+{
+	GENERATED_BODY()
+
+	UPROPERTY()
+	uint8 SectionIdx = 0;
+
+	UPROPERTY()
+	uint8 AnimSetIdx = 0;
+};
+
 UCLASS(meta = (BlueprintSpawnableComponent))
 class CONTEXTUALANIMATION_API UContextualAnimSceneActorComponent : public UPrimitiveComponent, public IIKGoalCreatorInterface
 {
@@ -45,6 +97,8 @@ public:
 
 	virtual void AddIKGoals_Implementation(TMap<FName, FIKRigGoal>& OutGoals) override;
 
+	const FContextualAnimSceneBindings& GetBindings() const { return Bindings; };
+
 	/** Called when the actor owner of this component joins an scene */
 	UFUNCTION(BlueprintCallable, Category = "Contextual Anim|Scene Actor Component")
 	void OnJoinedScene(const FContextualAnimSceneBindings& InBindings);
@@ -62,6 +116,15 @@ public:
 	UFUNCTION(BlueprintCallable, Category = "Contextual Anim|Scene Actor Component")
 	bool StartContextualAnimScene(const FContextualAnimSceneBindings& InBindings);
 
+	UFUNCTION(BlueprintCallable, Category = "Contextual Anim|Scene Actor Component")
+	bool LateJoinContextualAnimScene(AActor* Actor, FName Role);
+
+	UFUNCTION(BlueprintCallable, Category = "Contextual Anim|Scene Actor Component")
+	bool TransitionContextualAnimScene(FName SectionName);
+
+	UFUNCTION(BlueprintCallable, Category = "Contextual Anim|Scene Actor Component")
+	void PlayExternalAnimation(UAnimSequenceBase* Animation);
+
 	void EarlyOutContextualAnimScene();
 
 	bool IsOwnerLocallyControlled() const;
@@ -75,6 +138,15 @@ protected:
 	UPROPERTY(Transient, ReplicatedUsing = OnRep_Bindings)
 	FContextualAnimSceneBindings RepBindings;
 
+	UPROPERTY(Transient, ReplicatedUsing = OnRep_LateJoinData)
+	FContextualAnimRepLateJoinData RepLateJoinData;
+
+	UPROPERTY(Transient, ReplicatedUsing = OnRep_PlayAnimData)
+	FContextualAnimRepPlayAnimData RepPlayAnimData;
+
+	UPROPERTY(Transient, ReplicatedUsing = OnRep_TransitionData)
+	FContextualAnimRepTransitionData RepTransitionData;
+
 	/**
 	 * Bindings for the interaction we are currently playing.
 	 * Used to update IK, keep montage in sync, disable/enable collision between actors etc
@@ -107,8 +179,20 @@ protected:
 	UFUNCTION()
 	void OnRep_Bindings();
 
+	UFUNCTION()
+	void OnRep_LateJoinData();
+
+	UFUNCTION()
+	void OnRep_PlayAnimData();
+
+	UFUNCTION()
+	void OnRep_TransitionData();
+
 	void SetIgnoreCollisionWithOtherActors(bool bValue) const;
 
+	void SetMovementState(bool bRequireFlyingMode);
+	void RestoreMovementState(bool bRequireFlyingMode);
+
 	UFUNCTION()
 	void OnMontageBlendingOut(UAnimMontage* Montage, bool bInterrupted);
 
@@ -119,16 +203,28 @@ protected:
 
 	void LeaveScene();
 
+	void LateJoinScene(const FContextualAnimSceneBindings& InBindings);
+
+	bool HandleLateJoin(AActor* Actor, FName Role);
+
 	UFUNCTION(Server, Reliable, WithValidation)
 	void ServerStartContextualAnimScene(const FContextualAnimSceneBindings& InBindings);
 
 	UFUNCTION(Server, Reliable, WithValidation)
 	void ServerEarlyOutContextualAnimScene();
 
+	void PlayAnimation_Internal(UAnimSequenceBase* Animation, float StartTime, bool bSyncPlaybackTime);
+
+	void AddOrUpdateWarpTargets(int32 SectionIdx, int32 AnimSetIdx);
+
+	void HandleTransitionSelf(int32 NewSectionIdx, int32 NewAnimSetIdx);
+
+	void HandleTransitionEveryone(int32 NewSectionIdx, int32 NewAnimSetIdx);
+
 private:
 
 	bool bRegistered = false;
-
+	bool bGuardAnimEvents = false;
 };
 
 #if UE_ENABLE_INCLUDE_ORDER_DEPRECATED_IN_5_2
diff --git a/Engine/Plugins/Experimental/Animation/ContextualAnimation/Source/ContextualAnimation/Public/ContextualAnimTypes.h b/Engine/Plugins/Experimental/Animation/ContextualAnimation/Source/ContextualAnimation/Public/ContextualAnimTypes.h
--- a/Engine/Plugins/Experimental/Animation/ContextualAnimation/Source/ContextualAnimation/Public/ContextualAnimTypes.h
+++ b/Engine/Plugins/Experimental/Animation/ContextualAnimation/Source/ContextualAnimation/Public/ContextualAnimTypes.h
@@ -421,6 +421,7 @@ struct CONTEXTUALANIMATION_API FContextualAnimSceneBindings
 	FTransform GetIKTargetTransformFromBinding(const FContextualAnimSceneBinding& Binding, const FName& TrackName, float Time) const;
 
 	const FContextualAnimSceneBinding* GetSyncLeader() const;
+	const FContextualAnimSceneBinding* GetPrimaryBinding() const;
 
 	bool IsValid() const;
 
@@ -430,6 +431,8 @@ struct CONTEXTUALANIMATION_API FContextualAnimSceneBindings
 
 	void GenerateUniqueId();
 
+	void TransitionTo(int32 NewSectionIdx, int32 NewAnimSetIdx);
+
 private:
 
 	friend class UContextualAnimManager;
