Contextual Anim: Added feature to allow actors to define warp points using sockets in their mesh. When those are present alignment will be calculate relative to those points. This is useful for cases where we want to re-use the same animation to interact with props with different proportions.

diff --git a/Engine/Plugins/Experimental/Animation/ContextualAnimation/Source/ContextualAnimation/Private/ContextualAnimSceneActorComponent.cpp b/Engine/Plugins/Experimental/Animation/ContextualAnimation/Source/ContextualAnimation/Private/ContextualAnimSceneActorComponent.cpp
--- a/Engine/Plugins/Experimental/Animation/ContextualAnimation/Source/ContextualAnimation/Private/ContextualAnimSceneActorComponent.cpp
+++ b/Engine/Plugins/Experimental/Animation/ContextualAnimation/Source/ContextualAnimation/Private/ContextualAnimSceneActorComponent.cpp
@@ -115,23 +115,82 @@ void UContextualAnimSceneActorComponent::PlayAnimation_Internal(UAnimSequenceBas
 
 void UContextualAnimSceneActorComponent::AddOrUpdateWarpTargets(int32 SectionIdx, int32 AnimSetIdx)
 {
+	// This is relevant only for character with motion warping comp
+	ACharacter* CharacterOwner = Cast<ACharacter>(GetOwner());
+	UMotionWarpingComponent* MotionWarpComp = CharacterOwner ? CharacterOwner->GetComponentByClass<UMotionWarpingComponent>() : nullptr;
+	if (MotionWarpComp == nullptr)
+	{
+		return;
+	}
+
 	if (const FContextualAnimSceneBinding* Binding = Bindings.FindBindingByActor(GetOwner()))
 	{
 		const UContextualAnimSceneAsset* Asset = Bindings.GetSceneAsset();
 		check(Asset);
 
-		const TArray<FContextualAnimSetPivotDefinition>& PivotDefs = Asset->GetAnimSetPivotDefinitionsInSection(SectionIdx);
-		if (PivotDefs.Num() > 0)
+		// If WarpPoints array is not empty means that its the primary actor who will provide the transforms for alignment
+		if(Asset->GetWarpPoints().Num() > 0)
 		{
-			const FContextualAnimTrack* AnimTrack = Asset->GetAnimTrack(SectionIdx, AnimSetIdx, Bindings.GetRoleFromBinding(*Binding));
-			check(AnimTrack);
+			const FContextualAnimSceneBinding* PrimaryBinding = Bindings.GetPrimaryBinding();
+			if(PrimaryBinding == nullptr)
+			{
+				return;
+			}
+
+			for (const FContextualAnimWarpPointData& WarpPointData : Asset->GetWarpPoints())
+			{
+				// Get original warp point cached in the SceneAsset
+				const FTransform* WarpPointTransform = Asset->GetWarpPointTransform(WarpPointData.SocketName);
+				if (WarpPointTransform == nullptr)
+				{
+					continue;
+				}
+
+				// Put original warp point in world space
+				const FTransform OriginalWarpPointTransform = *WarpPointTransform * PrimaryBinding->GetTransform();
 
-			for (const FContextualAnimSetPivotDefinition& PivotDef : PivotDefs)
+				// Now pull the warp point transform from the actual actor. This transform could be different if interacting with actors with different proportions
+				FTransform FinalWarpPointTransform = OriginalWarpPointTransform;
+				if (UMeshComponent* Component = UContextualAnimUtilities::TryGetMeshComponentWithSocket(PrimaryBinding->GetActor(), WarpPointData.SocketName))
+				{
+					const FTransform SocketTransform = Component->GetSocketTransform(WarpPointData.SocketName, ERelativeTransformSpace::RTS_Actor);
+					FinalWarpPointTransform = SocketTransform * PrimaryBinding->GetTransform();
+				}
+
+				// Get the NotifyEvent from the warping window for the corresponding warp target
+				const FContextualAnimTrack* AnimTrack = Asset->GetAnimTrack(SectionIdx, AnimSetIdx, Bindings.GetRoleFromBinding(*Binding));
+				const FAnimNotifyEvent* NotifyEvent = (AnimTrack && AnimTrack->Animation) ? UContextualAnimUtilities::FindFirstWarpingWindowForWarpTarget(AnimTrack->Animation, WarpPointData.WarpTargetName) : nullptr;
+				if (NotifyEvent)
+				{
+					// Calculate offset between the original warp point and the actual warp point
+					const FTransform Offset = FinalWarpPointTransform * OriginalWarpPointTransform.Inverse();
+
+					// Get where the actor ends in the animation at the end of this warping window, apply the offset and put it in world space
+					// This assumes the primary actor is at the origin of the scene
+					FTransform WarpTargetTransform = AnimTrack->GetRootTransformAtTime(NotifyEvent->GetEndTriggerTime());
+					WarpTargetTransform = (WarpTargetTransform * Offset) * PrimaryBinding->GetTransform();
+
+					// Account for mesh offset in character blueprint
+					WarpTargetTransform = FTransform(CharacterOwner->GetBaseRotationOffset().Inverse()) * WarpTargetTransform;
+
+					// Add warp target transform
+					MotionWarpComp->AddOrUpdateWarpTargetFromTransform(WarpPointData.WarpTargetName, WarpTargetTransform);
+				}
+			}
+		}
+		// When warp points are not explicitly defined by the primary actor we generate them on the fly based on AnimSetPivotDefinitions
+		else 
+		{
+			const TArray<FContextualAnimSetPivotDefinition>& PivotDefs = Asset->GetAnimSetPivotDefinitionsInSection(SectionIdx);
+			if (PivotDefs.Num() > 0)
 			{
-				FContextualAnimSetPivot ScenePivot;
-				if (Bindings.CalculateAnimSetPivot(PivotDef, ScenePivot))
+				const FContextualAnimTrack* AnimTrack = Asset->GetAnimTrack(SectionIdx, AnimSetIdx, Bindings.GetRoleFromBinding(*Binding));
+				check(AnimTrack);
+
+				for (const FContextualAnimSetPivotDefinition& PivotDef : PivotDefs)
 				{
-					if (UMotionWarpingComponent* MotionWarpComp = GetOwner()->FindComponentByClass<UMotionWarpingComponent>())
+					FContextualAnimSetPivot ScenePivot;
+					if (Bindings.CalculateAnimSetPivot(PivotDef, ScenePivot))
 					{
 						const float Time = AnimTrack->GetSyncTimeForWarpSection(PivotDef.Name);
 						const FTransform TransformRelativeToScenePivot = Asset->GetAlignmentTransform(SectionIdx, AnimSetIdx, AnimTrack->AnimTrackIdx, PivotDef.Name, Time);
diff --git a/Engine/Plugins/Experimental/Animation/ContextualAnimation/Source/ContextualAnimation/Private/ContextualAnimUtilities.cpp b/Engine/Plugins/Experimental/Animation/ContextualAnimation/Source/ContextualAnimation/Private/ContextualAnimUtilities.cpp
--- a/Engine/Plugins/Experimental/Animation/ContextualAnimation/Source/ContextualAnimation/Private/ContextualAnimUtilities.cpp
+++ b/Engine/Plugins/Experimental/Animation/ContextualAnimation/Source/ContextualAnimation/Private/ContextualAnimUtilities.cpp
@@ -11,8 +11,11 @@
 #include "GameFramework/Character.h"
 #include "ContextualAnimActorInterface.h"
 #include "Components/SkeletalMeshComponent.h"
+#include "Components/StaticMeshComponent.h"
 #include "SceneManagement.h"
 #include "MotionWarpingComponent.h"
+#include "AnimNotifyState_MotionWarping.h"
+#include "RootMotionModifier.h"
 
 #include UE_INLINE_GENERATED_CPP_BY_NAME(ContextualAnimUtilities)
 
@@ -177,6 +180,47 @@ void UContextualAnimUtilities::DrawDebugAnimSet(const UWorld* World, const UCont
 	}
 }
 
+const FAnimNotifyEvent* UContextualAnimUtilities::FindFirstWarpingWindowForWarpTarget(const UAnimSequenceBase* Animation, FName WarpTargetName)
+{
+	if(Animation)
+	{
+		return Animation->Notifies.FindByPredicate([WarpTargetName](const FAnimNotifyEvent& NotifyEvent)
+		{
+			if (const UAnimNotifyState_MotionWarping* Notify = Cast<UAnimNotifyState_MotionWarping>(NotifyEvent.NotifyStateClass))
+			{
+				if (const URootMotionModifier_Warp* Modifier = Cast<URootMotionModifier_Warp>(Notify->RootMotionModifier))
+				{
+					if (Modifier->WarpTargetName == WarpTargetName)
+					{
+						return true;
+					}
+				}
+			}
+
+			return false;
+		});
+	}
+
+	return nullptr;
+}
+
+UMeshComponent* UContextualAnimUtilities::TryGetMeshComponentWithSocket(const AActor* Actor, FName SocketName)
+{
+	if (Actor)
+	{
+		TInlineComponentArray<UMeshComponent*> Components(Actor);
+		for (UMeshComponent* Component : Components)
+		{
+			if (Component && Component->DoesSocketExist(SocketName))
+			{
+				return Component;
+			}
+		}
+	}
+
+	return nullptr;
+}
+
 USkeletalMeshComponent* UContextualAnimUtilities::TryGetSkeletalMeshComponent(const AActor* Actor)
 {
 	USkeletalMeshComponent* SkelMeshComp = nullptr;
diff --git a/Engine/Plugins/Experimental/Animation/ContextualAnimation/Source/ContextualAnimation/Public/ContextualAnimSceneAsset.h b/Engine/Plugins/Experimental/Animation/ContextualAnimation/Source/ContextualAnimation/Public/ContextualAnimSceneAsset.h
--- a/Engine/Plugins/Experimental/Animation/ContextualAnimation/Source/ContextualAnimation/Public/ContextualAnimSceneAsset.h
+++ b/Engine/Plugins/Experimental/Animation/ContextualAnimation/Source/ContextualAnimation/Public/ContextualAnimSceneAsset.h
@@ -55,6 +55,20 @@ struct CONTEXTUALANIMATION_API FContextualAnimSet
 	int32 GetNumMandatoryRoles() const;
 };
 
+USTRUCT(BlueprintType)
+struct CONTEXTUALANIMATION_API FContextualAnimWarpPointData
+{
+	GENERATED_BODY()
+
+	/** Name of the socket in the primary actor acting as warping point (reference point for alignment) */
+	UPROPERTY(EditAnywhere, Category = "Defaults")
+	FName SocketName = NAME_None;
+
+	/** Name of the warp target in the warping window this warp point is for */
+	UPROPERTY(EditAnywhere, Category = "Defaults")
+	FName WarpTargetName = NAME_None;
+};
+
 /** Named container with one or more ContextualAnimSet */
 USTRUCT(BlueprintType)
 struct CONTEXTUALANIMATION_API FContextualAnimSceneSection
@@ -207,6 +221,8 @@ public:
 	FORCEINLINE const TSubclassOf<UContextualAnimSceneInstance>& GetSceneInstanceClass() const { return SceneInstanceClass; }
 	FORCEINLINE int32 GetSampleRate() const { return SampleRate; }
 	FORCEINLINE float GetRadius() const { return Radius; }
+	FORCEINLINE const TArray<FContextualAnimWarpPointData>& GetWarpPoints() const { return WarpPoints; }
+	FORCEINLINE const FTransform* GetWarpPointTransform(FName SocketName) const { return WarpPointTransformCache.Find(SocketName); }
 
 	bool HasValidData() const { return RolesAsset != nullptr && Sections.Num() > 0 && Sections[0].AnimSets.Num() > 0; }
 
@@ -292,6 +308,21 @@ protected:
 	UPROPERTY(EditAnywhere, Category = "Settings", meta = (TitleProperty = "Role"))
 	TArray<FContextualAnimActorPreviewData> OverridePreviewData;
 
+	/**
+	 * List of warp points defined with sockets by the the primary actor.
+	 * When this list is non empty the system uses these sockets as warp points.
+	 * Commonly used when re-using the same animations to interact with objects with different proportions
+	*/
+	UPROPERTY(EditAnywhere, Category = "Warping", meta = (TitleProperty = "WarpTargetName"))
+	TArray<FContextualAnimWarpPointData> WarpPoints;
+
+	/** 
+	 * Transforms of the sockets used as warp points in actor space. Updated in the editor via 'Update Warp Points button' 
+	 * We are caching this transforms here at least for now so we don't have to include a reference to the primary actor 
+	 */
+	UPROPERTY()
+	TMap<FName, FTransform> WarpPointTransformCache;
+
 	UPROPERTY(EditAnywhere, Category = "Defaults")
 	TArray<FContextualAnimSceneSection> Sections;
 
diff --git a/Engine/Plugins/Experimental/Animation/ContextualAnimation/Source/ContextualAnimation/Public/ContextualAnimUtilities.h b/Engine/Plugins/Experimental/Animation/ContextualAnimation/Source/ContextualAnimation/Public/ContextualAnimUtilities.h
--- a/Engine/Plugins/Experimental/Animation/ContextualAnimation/Source/ContextualAnimation/Public/ContextualAnimUtilities.h
+++ b/Engine/Plugins/Experimental/Animation/ContextualAnimation/Source/ContextualAnimation/Public/ContextualAnimUtilities.h
@@ -6,13 +6,16 @@
 #include "ContextualAnimUtilities.generated.h"
 
 class UContextualAnimSceneAsset;
+class UMeshComponent;
 class USkeletalMeshComponent;
+class UStaticMeshComponent;
 class UAnimInstance;
 class AActor;
 class FPrimitiveDrawInterface;
 struct FAnimMontageInstance;
 struct FCompactPose;
 struct FContextualAnimSet;
+struct FAnimNotifyEvent;
 template<class PoseType> struct FCSPose;
 
 UCLASS()
@@ -52,6 +55,10 @@ public:
 	
 	static void DrawDebugAnimSet(const UWorld* World, const UContextualAnimSceneAsset& SceneAsset, const FContextualAnimSet& AnimSet, float Time, const FTransform& ToWorldTransform, const FColor& Color, float LifeTime, float Thickness);
 
+	static const FAnimNotifyEvent* FindFirstWarpingWindowForWarpTarget(const UAnimSequenceBase* Animation, FName WarpTargetName);
+
+	static UMeshComponent* TryGetMeshComponentWithSocket(const AActor* Actor, FName SocketName);
+
 	static USkeletalMeshComponent* TryGetSkeletalMeshComponent(const AActor* Actor);
 
 	static UAnimInstance* TryGetAnimInstance(const AActor* Actor);
diff --git a/Engine/Plugins/Experimental/Animation/ContextualAnimation/Source/ContextualAnimationEditor/Private/ContextualAnimEdMode.cpp b/Engine/Plugins/Experimental/Animation/ContextualAnimation/Source/ContextualAnimationEditor/Private/ContextualAnimEdMode.cpp
--- a/Engine/Plugins/Experimental/Animation/ContextualAnimation/Source/ContextualAnimationEditor/Private/ContextualAnimEdMode.cpp
+++ b/Engine/Plugins/Experimental/Animation/ContextualAnimation/Source/ContextualAnimationEditor/Private/ContextualAnimEdMode.cpp
@@ -66,7 +66,53 @@ void FContextualAnimEdMode::Render(const FSceneView* View, FViewport* Viewport,
 				{
 					for (const FTransform& ScenePivot : AnimSet->ScenePivots)
 					{
-						DrawCoordinateSystem(PDI, ScenePivot.GetLocation(), ScenePivot.Rotator(), 50.f, SDPG_Foreground);
+						DrawCoordinateSystem(PDI, ScenePivot.GetLocation(), ScenePivot.Rotator(), 10.f, SDPG_Foreground, 0.5f);
+					}
+				}
+
+				// Draw warp points
+				for (const FContextualAnimWarpPointData& WarpPointData : SceneAsset->GetWarpPoints())
+				{
+					const FTransform* WarpPointTransform = SceneAsset->GetWarpPointTransform(WarpPointData.SocketName);
+					if (WarpPointTransform == nullptr)
+					{
+						continue;
+					}
+
+					const FTransform OriginalWarpPointTransform = *WarpPointTransform * PrimaryBinding->GetTransform();
+
+					// Draw original warp point in red.
+					const float Thickness = 0.f;
+					DrawCoordinateSystem(PDI, OriginalWarpPointTransform.GetLocation(), OriginalWarpPointTransform.Rotator(), 5.f, SDPG_Foreground, Thickness);
+					DrawWireDiamond(PDI, OriginalWarpPointTransform.ToMatrixNoScale(), 3.f, FLinearColor::Red, 0, Thickness);
+
+					// Draw actual warp point in green.
+					FTransform FinalWarpPointTransform = OriginalWarpPointTransform;
+					if (UMeshComponent* Component = UContextualAnimUtilities::TryGetMeshComponentWithSocket(PrimaryBinding->GetActor(), WarpPointData.SocketName))
+					{
+						const FTransform SocketTransform = Component->GetSocketTransform(WarpPointData.SocketName, ERelativeTransformSpace::RTS_Actor);
+						FinalWarpPointTransform = SocketTransform * PrimaryBinding->GetTransform();
+
+						DrawCoordinateSystem(PDI, FinalWarpPointTransform.GetLocation(), FinalWarpPointTransform.Rotator(), 5.f, SDPG_Foreground, Thickness);
+						DrawWireDiamond(PDI, FinalWarpPointTransform.ToMatrixNoScale(), 3.f, FLinearColor::Green, 0, Thickness);
+						PDI->DrawLine(OriginalWarpPointTransform.GetLocation(), FinalWarpPointTransform.GetLocation(), FLinearColor::Green, 0, 0.f);
+					}
+
+					// Draw warp target transforms for each actor using this warp point in Blue. This is where the root should end up at the end of the warping window.
+					for (const FContextualAnimSceneBinding& Binding : Bindings)
+					{
+						const FContextualAnimTrack& AnimTrack = Bindings.GetAnimTrackFromBinding(Binding);
+						const FAnimNotifyEvent* NotifyEvent = (AnimTrack.Animation) ? UContextualAnimUtilities::FindFirstWarpingWindowForWarpTarget(AnimTrack.Animation, WarpPointData.WarpTargetName) : nullptr;
+
+						if (NotifyEvent)
+						{
+							const FTransform Offset = FinalWarpPointTransform * OriginalWarpPointTransform.Inverse();
+							FTransform AlignmentTransform = AnimTrack.GetRootTransformAtTime(NotifyEvent->GetEndTriggerTime());
+							AlignmentTransform = (AlignmentTransform * Offset) * PrimaryBinding->GetTransform();
+
+							DrawCoordinateSystem(PDI, AlignmentTransform.GetLocation(), AlignmentTransform.Rotator(), 5.f, SDPG_Foreground, Thickness);
+							DrawWireDiamond(PDI, AlignmentTransform.ToMatrixNoScale(), 3.f, FLinearColor::Blue, 0, Thickness);
+						}
 					}
 				}
 
diff --git a/Engine/Plugins/Experimental/Animation/ContextualAnimation/Source/ContextualAnimationEditor/Private/ContextualAnimViewModel.cpp b/Engine/Plugins/Experimental/Animation/ContextualAnimation/Source/ContextualAnimationEditor/Private/ContextualAnimViewModel.cpp
--- a/Engine/Plugins/Experimental/Animation/ContextualAnimation/Source/ContextualAnimationEditor/Private/ContextualAnimViewModel.cpp
+++ b/Engine/Plugins/Experimental/Animation/ContextualAnimation/Source/ContextualAnimationEditor/Private/ContextualAnimViewModel.cpp
@@ -1227,6 +1227,23 @@ void FContextualAnimViewModel::DiscardChangeToActorTransformInScene()
 	ModifyingTransformInSceneCachedActor.Reset();
 }
 
+void FContextualAnimViewModel::CacheWarpPoints()
+{
+	SceneAsset->WarpPointTransformCache.Reset();
+
+	for (const FContextualAnimWarpPointData& WarpPointData : SceneAsset->WarpPoints)
+	{
+		if (const FContextualAnimSceneBinding* Primary = SceneBindings.GetPrimaryBinding())
+		{
+			if (UMeshComponent* MeshComp = UContextualAnimUtilities::TryGetMeshComponentWithSocket(Primary->GetActor(), WarpPointData.SocketName))
+			{
+				const FTransform WarpPointTransform = MeshComp->GetSocketTransform(WarpPointData.SocketName, ERelativeTransformSpace::RTS_Actor);
+				SceneAsset->WarpPointTransformCache.Add(WarpPointData.SocketName, WarpPointTransform);
+			}
+		}
+	}
+}
+
 void FContextualAnimViewModel::OnFinishedChangingProperties(const FPropertyChangedEvent& PropertyChangedEvent)
 {
 	const FName PropertyName = (PropertyChangedEvent.Property != nullptr) ? PropertyChangedEvent.Property->GetFName() : NAME_None;
diff --git a/Engine/Plugins/Experimental/Animation/ContextualAnimation/Source/ContextualAnimationEditor/Private/DetailCustomizations/ContextualAnimSceneAssetDetailCustom.cpp b/Engine/Plugins/Experimental/Animation/ContextualAnimation/Source/ContextualAnimationEditor/Private/DetailCustomizations/ContextualAnimSceneAssetDetailCustom.cpp
--- a/Engine/Plugins/Experimental/Animation/ContextualAnimation/Source/ContextualAnimationEditor/Private/DetailCustomizations/ContextualAnimSceneAssetDetailCustom.cpp
+++ b/Engine/Plugins/Experimental/Animation/ContextualAnimation/Source/ContextualAnimationEditor/Private/DetailCustomizations/ContextualAnimSceneAssetDetailCustom.cpp
@@ -33,6 +33,32 @@ void FContextualAnimSceneAssetDetailCustom::CustomizeDetails(IDetailLayoutBuilde
 	const UContextualAnimSceneAsset* SceneAsset = Cast<const UContextualAnimSceneAsset>(ObjectList[0].Get());
 	check(SceneAsset);
 
+	IDetailCategoryBuilder& SettingsCategory = DetailBuilder.EditCategory(TEXT("Settings"));
+	SettingsCategory.SetSortOrder(0);
+
+	IDetailCategoryBuilder& WarpingCategory = DetailBuilder.EditCategory(TEXT("Warping"));
+	WarpingCategory.SetSortOrder(1);
+
+	// Add a custom button to extract the warp points from the primary actor and cache them in the SceneAsset.
+	// @TODO: This is probably temp to get the feature in. It may be better to just add a reference to the original actor in the asset. Workflow will be much better that way. 
+	WarpingCategory.AddCustomRow(FText::GetEmpty())
+		.ValueContent()
+		.VAlign(VAlign_Center)
+		.MaxDesiredWidth(250)
+		[
+			SNew(SButton).Text(LOCTEXT("UpdateWarpPointsLabel", "Update Warp Points"))
+			.OnClicked_Lambda([ViewModel]()
+				{
+					const FText DialogMsg = LOCTEXT("RefreshWarpPointsDialog", "Warp Points should be updated while previewing the scene with the meshes the animations where originally authored for. Are you sure you want to continue?");
+					if (FMessageDialog::Open(EAppMsgType::YesNo, DialogMsg) == EAppReturnType::Yes)
+					{
+						ViewModel->CacheWarpPoints();
+					}
+
+					return FReply::Handled();
+				})
+		];
+
 	IDetailCategoryBuilder& Category = DetailBuilder.EditCategory(TEXT("Defaults"));
 	Category.SetCategoryVisibility(false);
 
diff --git a/Engine/Plugins/Experimental/Animation/ContextualAnimation/Source/ContextualAnimationEditor/Public/ContextualAnimViewModel.h b/Engine/Plugins/Experimental/Animation/ContextualAnimation/Source/ContextualAnimationEditor/Public/ContextualAnimViewModel.h
--- a/Engine/Plugins/Experimental/Animation/ContextualAnimation/Source/ContextualAnimationEditor/Public/ContextualAnimViewModel.h
+++ b/Engine/Plugins/Experimental/Animation/ContextualAnimation/Source/ContextualAnimationEditor/Public/ContextualAnimViewModel.h
@@ -129,6 +129,8 @@ public:
 	
 	bool CanMakeEdits() const;
 
+	void CacheWarpPoints();
+
 private:
 
 	struct FContextualAnimEdSelectionInfo
