Contextual Anim: Refactored option to play arbitrary animation on actor so you can only play animations in the scene asset for the current interactions and now warping targets are properly calculated and added for this animation too.

diff --git a/Engine/Plugins/Experimental/Animation/ContextualAnimation/Source/ContextualAnimation/Private/ContextualAnimSceneActorComponent.cpp b/Engine/Plugins/Experimental/Animation/ContextualAnimation/Source/ContextualAnimation/Private/ContextualAnimSceneActorComponent.cpp
--- a/Engine/Plugins/Experimental/Animation/ContextualAnimation/Source/ContextualAnimation/Private/ContextualAnimSceneActorComponent.cpp
+++ b/Engine/Plugins/Experimental/Animation/ContextualAnimation/Source/ContextualAnimation/Private/ContextualAnimSceneActorComponent.cpp
@@ -55,7 +55,7 @@ void UContextualAnimSceneActorComponent::GetLifetimeReplicatedProps(TArray< FLif
 	Params.bIsPushBased = true;
 	DOREPLIFETIME_WITH_PARAMS_FAST(UContextualAnimSceneActorComponent, RepBindings, Params);
 	DOREPLIFETIME_WITH_PARAMS_FAST(UContextualAnimSceneActorComponent, RepLateJoinData, Params);
-	DOREPLIFETIME_WITH_PARAMS_FAST(UContextualAnimSceneActorComponent, RepPlayAnimData, Params);
+	DOREPLIFETIME_WITH_PARAMS_FAST(UContextualAnimSceneActorComponent, RepTransitionSingleActorData, Params);
 	DOREPLIFETIME_WITH_PARAMS_FAST(UContextualAnimSceneActorComponent, RepTransitionData, Params);
 }
 
@@ -406,27 +406,58 @@ void UContextualAnimSceneActorComponent::HandleTransitionSelf(int32 NewSectionId
 	AddOrUpdateWarpTargets(NewSectionIdx, NewAnimSetIdx);
 }
 
-void UContextualAnimSceneActorComponent::PlayExternalAnimation(UAnimSequenceBase* Animation)
+bool UContextualAnimSceneActorComponent::TransitionSingleActor(int32 SectionIdx, int32 AnimSetIdx)
 {
 	if (!GetOwner()->HasAuthority())
 	{
-		return;
+		return false;
 	}
 
-	PlayAnimation_Internal(Animation, 0.f, false);
+	if (const FContextualAnimSceneBinding* OwnerBinding = Bindings.FindBindingByActor(GetOwner()))
+	{
+		if (const UContextualAnimSceneAsset* Asset = Bindings.GetSceneAsset())
+		{
+			const FContextualAnimTrack* AnimTrack = Asset->GetAnimTrack(SectionIdx, AnimSetIdx, Bindings.GetRoleFromBinding(*OwnerBinding));
+			if (AnimTrack && AnimTrack->Animation)
+			{
+				UE_LOG(LogContextualAnim, Log, TEXT("%-21s UContextualAnimSceneActorComponent::TransitionSingleActor Actor: %s SectionIdx: %d AnimSetIdx: %d"),
+					*UEnum::GetValueAsString(TEXT("Engine.ENetRole"), GetOwner()->GetLocalRole()), *GetNameSafe(GetOwner()), SectionIdx, AnimSetIdx);
+
+				PlayAnimation_Internal(AnimTrack->Animation, 0.f, false);
+
+				AddOrUpdateWarpTargets(SectionIdx, AnimSetIdx);
+
+				RepTransitionSingleActorData.SectionIdx = SectionIdx;
+				RepTransitionSingleActorData.AnimSetIdx = AnimSetIdx;
+				RepTransitionSingleActorData.IncrementRepCounter();
+				MARK_PROPERTY_DIRTY_FROM_NAME(UContextualAnimSceneActorComponent, RepTransitionSingleActorData, this);
+				GetOwner()->ForceNetUpdate();
+
+				return true;
+			}
+		}
+	}
 
-	RepPlayAnimData.Animation = Animation;
-	RepPlayAnimData.IncrementRepCounter();
-	MARK_PROPERTY_DIRTY_FROM_NAME(UContextualAnimSceneActorComponent, RepPlayAnimData, this);
-	GetOwner()->ForceNetUpdate();
+	return false;
 }
 
-void UContextualAnimSceneActorComponent::OnRep_PlayAnimData()
+void UContextualAnimSceneActorComponent::OnRep_RepTransitionSingleActor()
 {
-	UE_LOG(LogContextualAnim, Verbose, TEXT("%-21s UContextualAnimSceneActorComponent::OnRep_PlayAnimData Owner: %s AnimToPlay: %s"),
-		*UEnum::GetValueAsString(TEXT("Engine.ENetRole"), GetOwner()->GetLocalRole()), *GetNameSafe(GetOwner()), *GetNameSafe(RepPlayAnimData.Animation));
+	UE_LOG(LogContextualAnim, Verbose, TEXT("%-21s UContextualAnimSceneActorComponent::OnRep_RepTransitionSingleActor Owner: %s SectionIdx: %s AnimSetIdx: %d"),
+		*UEnum::GetValueAsString(TEXT("Engine.ENetRole"), GetOwner()->GetLocalRole()), *GetNameSafe(GetOwner()), RepTransitionSingleActorData.SectionIdx, RepTransitionSingleActorData.AnimSetIdx);
 
-	PlayAnimation_Internal(RepPlayAnimData.Animation, 0.f, false);
+	if (const FContextualAnimSceneBinding* OwnerBinding = Bindings.FindBindingByActor(GetOwner()))
+	{
+		if (const UContextualAnimSceneAsset* Asset = Bindings.GetSceneAsset())
+		{
+			const FContextualAnimTrack* AnimTrack = Asset->GetAnimTrack(RepTransitionSingleActorData.SectionIdx, RepTransitionSingleActorData.AnimSetIdx, Bindings.GetRoleFromBinding(*OwnerBinding));
+			if (AnimTrack && AnimTrack->Animation)
+			{
+				PlayAnimation_Internal(AnimTrack->Animation, 0.f, false);
+				AddOrUpdateWarpTargets(RepTransitionSingleActorData.SectionIdx, RepTransitionSingleActorData.AnimSetIdx);
+			}
+		}
+	}
 }
 
 bool UContextualAnimSceneActorComponent::StartContextualAnimScene(const FContextualAnimSceneBindings& InBindings)
diff --git a/Engine/Plugins/Experimental/Animation/ContextualAnimation/Source/ContextualAnimation/Public/ContextualAnimSceneActorComponent.h b/Engine/Plugins/Experimental/Animation/ContextualAnimation/Source/ContextualAnimation/Public/ContextualAnimSceneActorComponent.h
--- a/Engine/Plugins/Experimental/Animation/ContextualAnimation/Source/ContextualAnimation/Public/ContextualAnimSceneActorComponent.h
+++ b/Engine/Plugins/Experimental/Animation/ContextualAnimation/Source/ContextualAnimation/Public/ContextualAnimSceneActorComponent.h
@@ -43,16 +43,6 @@ struct FContextualAnimRepLateJoinData : public FContextualAnimRepData
 	FName Role = NAME_None;
 };
 
-/** Used to replicate the playback of a 'random' animation */
-USTRUCT()
-struct FContextualAnimRepPlayAnimData : public FContextualAnimRepData
-{
-	GENERATED_BODY()
-
-	UPROPERTY()
-	TObjectPtr<UAnimSequenceBase> Animation = nullptr;
-};
-
 /** Used to replicate a transition to a new AnimSet in the interaction */
 USTRUCT()
 struct FContextualAnimRepTransitionData : public FContextualAnimRepData
@@ -123,7 +113,7 @@ public:
 	bool TransitionContextualAnimScene(FName SectionName);
 
 	UFUNCTION(BlueprintCallable, Category = "Contextual Anim|Scene Actor Component")
-	void PlayExternalAnimation(UAnimSequenceBase* Animation);
+	bool TransitionSingleActor(int32 SectionIdx, int32 AnimSetIdx);
 
 	void EarlyOutContextualAnimScene();
 
@@ -141,12 +131,12 @@ protected:
 	UPROPERTY(Transient, ReplicatedUsing = OnRep_LateJoinData)
 	FContextualAnimRepLateJoinData RepLateJoinData;
 
-	UPROPERTY(Transient, ReplicatedUsing = OnRep_PlayAnimData)
-	FContextualAnimRepPlayAnimData RepPlayAnimData;
-
 	UPROPERTY(Transient, ReplicatedUsing = OnRep_TransitionData)
 	FContextualAnimRepTransitionData RepTransitionData;
 
+	UPROPERTY(Transient, ReplicatedUsing = OnRep_RepTransitionSingleActor)
+	FContextualAnimRepTransitionData RepTransitionSingleActorData;
+
 	/**
 	 * Bindings for the interaction we are currently playing.
 	 * Used to update IK, keep montage in sync, disable/enable collision between actors etc
@@ -183,7 +173,7 @@ protected:
 	void OnRep_LateJoinData();
 
 	UFUNCTION()
-	void OnRep_PlayAnimData();
+	void OnRep_RepTransitionSingleActor();
 
 	UFUNCTION()
 	void OnRep_TransitionData();
