Contextual Anim Editor: Fixed bug causing missalignment in the preview scene due to scene actor comp to being attached to the actors

diff --git a/Engine/Plugins/Experimental/Animation/ContextualAnimation/Source/ContextualAnimationEditor/Private/ContextualAnimViewModel.cpp b/Engine/Plugins/Experimental/Animation/ContextualAnimation/Source/ContextualAnimationEditor/Private/ContextualAnimViewModel.cpp
--- a/Engine/Plugins/Experimental/Animation/ContextualAnimation/Source/ContextualAnimationEditor/Private/ContextualAnimViewModel.cpp
+++ b/Engine/Plugins/Experimental/Animation/ContextualAnimation/Source/ContextualAnimationEditor/Private/ContextualAnimViewModel.cpp
@@ -292,6 +292,11 @@ AActor* FContextualAnimViewModel::SpawnPreviewActor(const FContextualAnimTrack&
 			MotionWarpingComp->RegisterComponentWithWorld(GetWorld());
 			MotionWarpingComp->InitializeComponent();
 
+			UContextualAnimSceneActorComponent* SceneActorComp = NewObject<UContextualAnimSceneActorComponent>(PreviewCharacter);
+			SceneActorComp->RegisterComponentWithWorld(GetWorld());
+			SceneActorComp->InitializeComponent();
+			SceneActorComp->AttachToComponent(PreviewCharacter->GetRootComponent(), FAttachmentTransformRules(EAttachmentRule::SnapToTarget, false));
+
 			PreviewActor = PreviewCharacter;
 		}
 		else
@@ -337,14 +342,14 @@ AActor* FContextualAnimViewModel::SpawnPreviewActor(const FContextualAnimTrack&
 					StaticMeshComp->RegisterComponentWithWorld(GetWorld());
 					PreviewActor->SetRootComponent(StaticMeshComp);
 				}
+
+				UContextualAnimSceneActorComponent* SceneActorComp = NewObject<UContextualAnimSceneActorComponent>(PreviewActor);
+				SceneActorComp->RegisterComponentWithWorld(GetWorld());
+				SceneActorComp->InitializeComponent();
+				SceneActorComp->AttachToComponent(PreviewActor->GetRootComponent(), FAttachmentTransformRules(EAttachmentRule::SnapToTarget, false));
 			}
 		}
 
-		check(PreviewActor);
-		UContextualAnimSceneActorComponent* SceneActorComp = NewObject<UContextualAnimSceneActorComponent>(PreviewActor);
-		SceneActorComp->RegisterComponentWithWorld(GetWorld());
-		SceneActorComp->InitializeComponent();
-
 		PreviewActor->Tags.Add(PreviewActorTag);
 	}
 
