Contextual Anim Editor: Actors in the preview scene are now aligned relative to the warp point defined by the primary actor (if any). This helps to prevent actors from overlapping each other when previewing the interaction with different meshes.

diff --git a/Engine/Plugins/Experimental/Animation/ContextualAnimation/Source/ContextualAnimationEditor/Private/ContextualAnimViewModel.cpp b/Engine/Plugins/Experimental/Animation/ContextualAnimation/Source/ContextualAnimationEditor/Private/ContextualAnimViewModel.cpp
--- a/Engine/Plugins/Experimental/Animation/ContextualAnimation/Source/ContextualAnimationEditor/Private/ContextualAnimViewModel.cpp
+++ b/Engine/Plugins/Experimental/Animation/ContextualAnimation/Source/ContextualAnimationEditor/Private/ContextualAnimViewModel.cpp
@@ -782,6 +782,37 @@ void FContextualAnimViewModel::UpdatePreviewActorTransform(const FContextualAnim
 			Transform = SceneBindings.GetAnimTrackFromBinding(Binding).GetRootTransformAtTime(Time);
 		}
 
+		// If custom warp point is defined by the primary actor, offset actors in the preview scene using the warp point as reference		
+		if (GetSceneAsset()->GetWarpPoints().Num() > 0)
+		{
+			// In this case everyone is positioned relative to the warp point defined by the primary actor
+			const FContextualAnimSceneBinding* PrimaryBinding = SceneBindings.GetPrimaryBinding();
+			if (PrimaryBinding && PrimaryBinding->GetActor() != Binding.GetActor())
+			{
+				// @TODO: For now we are always using the first warp point but we may want to let the user specify which warp point should be used for this with another option in the preview category.
+				const FContextualAnimWarpPointData& WarpPointData = GetSceneAsset()->GetWarpPoints()[0];
+
+				// Get original warp point cached in the SceneAsset
+				if (const FTransform* WarpPointTransform = GetSceneAsset()->GetWarpPointTransform(WarpPointData.SocketName))
+				{
+					// Put original warp point in world space
+					const FTransform OriginalWarpPointTransform = *WarpPointTransform * PrimaryBinding->GetTransform();
+
+					// Now pull the warp point transform from the actual actor. This transform could be different if interacting with actors with different proportions
+					FTransform FinalWarpPointTransform = OriginalWarpPointTransform;
+					if (UMeshComponent* Component = UContextualAnimUtilities::TryGetMeshComponentWithSocket(PrimaryBinding->GetActor(), WarpPointData.SocketName))
+					{
+						const FTransform SocketTransform = Component->GetSocketTransform(WarpPointData.SocketName, ERelativeTransformSpace::RTS_Actor);
+						FinalWarpPointTransform = SocketTransform * PrimaryBinding->GetTransform();
+					}
+
+					// Get offset between the actual warp point and the default one and apply it to the transform extracted from the animation
+					const FTransform Offset = FinalWarpPointTransform * OriginalWarpPointTransform.Inverse();
+					Transform = (Transform * Offset) * PrimaryBinding->GetTransform();
+				}
+			}
+		}
+
 		// Special case for Character
 		if (ACharacter* PreviewCharacter = Cast<ACharacter>(PreviewActor))
 		{
