Contextual Anim: Precomputing alignment tracks offline is now optional and false by default.

diff --git a/Engine/Plugins/Experimental/Animation/ContextualAnimation/Source/ContextualAnimation/Private/ContextualAnimSceneActorComponent.cpp b/Engine/Plugins/Experimental/Animation/ContextualAnimation/Source/ContextualAnimation/Private/ContextualAnimSceneActorComponent.cpp
--- a/Engine/Plugins/Experimental/Animation/ContextualAnimation/Source/ContextualAnimation/Private/ContextualAnimSceneActorComponent.cpp
+++ b/Engine/Plugins/Experimental/Animation/ContextualAnimation/Source/ContextualAnimation/Private/ContextualAnimSceneActorComponent.cpp
@@ -1042,11 +1042,11 @@ FPrimitiveSceneProxy* UContextualAnimSceneActorComponent::CreateSceneProxy()
 						if (AnimTrack.Role != SceneAssetPtr->GetPrimaryRole())
 						{
 							// Draw Entry Point
-							const FTransform EntryTransform = (AnimTrack.GetAlignmentTransformAtEntryTime() * ToWorldTransform);
+							const FTransform EntryTransform = (SceneAssetPtr->GetAlignmentTransform(AnimTrack, 0, 0.f) * ToWorldTransform);
 							DrawCoordinateSystem(PDI, EntryTransform.GetLocation(), EntryTransform.Rotator(), 20.f, SDPG_World, 3.f);
 
 							// Draw Sync Point
-							const FTransform SyncPoint = AnimTrack.GetAlignmentTransformAtSyncTime() * ToWorldTransform;
+							const FTransform SyncPoint = SceneAssetPtr->GetAlignmentTransform(AnimTrack, 0, AnimTrack.GetSyncTimeForWarpSection(0)) * ToWorldTransform;
 							DrawCoordinateSystem(PDI, SyncPoint.GetLocation(), SyncPoint.Rotator(), 20.f, SDPG_World, 3.f);
 
 							FLinearColor DrawColor = FLinearColor::White;
diff --git a/Engine/Plugins/Experimental/Animation/ContextualAnimation/Source/ContextualAnimation/Private/ContextualAnimSceneAsset.cpp b/Engine/Plugins/Experimental/Animation/ContextualAnimation/Source/ContextualAnimation/Private/ContextualAnimSceneAsset.cpp
--- a/Engine/Plugins/Experimental/Animation/ContextualAnimation/Source/ContextualAnimation/Private/ContextualAnimSceneAsset.cpp
+++ b/Engine/Plugins/Experimental/Animation/ContextualAnimation/Source/ContextualAnimation/Private/ContextualAnimSceneAsset.cpp
@@ -117,7 +117,23 @@ void FContextualAnimSceneSection::GenerateAlignmentTracks(UContextualAnimSceneAs
 
 			AnimSet.ScenePivots.Add(ScenePivot);
 		}
+	}
 
+	if(!SceneAsset.ShouldPrecomputeAlignmentTracks())
+	{
+		for (FContextualAnimSet& AnimSet : AnimSets)
+		{
+			for (FContextualAnimTrack& AnimTrack : AnimSet.Tracks)
+			{
+				AnimTrack.AlignmentData.Empty();
+			}
+		}
+
+		return;
+	}
+
+	for (FContextualAnimSet& AnimSet : AnimSets)
+	{
 		// Generate alignment tracks relative to scene pivot
 		for (FContextualAnimTrack& AnimTrack : AnimSet.Tracks)
 		{
@@ -387,23 +403,6 @@ const FContextualAnimTrack* FContextualAnimSceneSection::GetAnimTrack(int32 Anim
 	return nullptr;
 }
 
-FTransform FContextualAnimSceneSection::GetAlignmentTransformForRoleRelativeToPivot(int32 AnimSetIdx, FName Role, float Time) const
-{
-	if (const FContextualAnimTrack* AnimTrack = GetAnimTrack(AnimSetIdx, Role))
-	{
-		return AnimTrack->GetAlignmentTransformAtTime(Time);
-	}
-
-	return FTransform::Identity;
-}
-
-FTransform FContextualAnimSceneSection::GetAlignmentTransformForRoleRelativeToOtherRole(int32 AnimSetIdx, FName Role, FName OtherRole, float Time) const
-{
-	const FTransform FromRoleRelativeToScenePivot = GetAlignmentTransformForRoleRelativeToPivot(AnimSetIdx, Role, Time);
-	const FTransform ToRoleRelativeToScenePivot = GetAlignmentTransformForRoleRelativeToPivot(AnimSetIdx, OtherRole, Time);
-	return FromRoleRelativeToScenePivot.GetRelativeTransform(ToRoleRelativeToScenePivot);
-}
-
 FTransform FContextualAnimSceneSection::GetIKTargetTransformForRoleAtTime(int32 AnimSetIdx, FName Role, FName TrackName, float Time) const
 {
 	if (const FContextualAnimTrack* AnimTrack = GetAnimTrack(AnimSetIdx, Role))
@@ -453,31 +452,6 @@ const FContextualAnimTrack* FContextualAnimSceneSection::FindFirstAnimTrackForRo
 	return nullptr;
 }
 
-const FContextualAnimTrack* FContextualAnimSceneSection::FindAnimTrackForRoleWithClosestEntryLocation(const FName& Role, const FContextualAnimSceneBindingContext& Primary, const FVector& TestLocation) const
-{
-	const FContextualAnimTrack* Result = nullptr;
-
-	float BestDistanceSq = MAX_FLT;
-	for (const FContextualAnimSet& AnimSet : AnimSets)
-	{
-		for (const FContextualAnimTrack& AnimTrack : AnimSet.Tracks)
-		{
-			if (AnimTrack.Role == Role)
-			{
-				const FTransform EntryTransform = AnimTrack.GetAlignmentTransformAtEntryTime() * Primary.GetTransform();
-				const float DistSq = FVector::DistSquared(EntryTransform.GetLocation(), TestLocation);
-				if (DistSq < BestDistanceSq)
-				{
-					BestDistanceSq = DistSq;
-					Result = &AnimTrack;
-				}
-			}
-		}
-	}
-
-	return Result;
-}
-
 // UContextualAnimSceneAsset
 //==============================================================================================
 
@@ -650,13 +624,13 @@ bool UContextualAnimSceneAsset::Query(FName Role, FContextualAnimQueryResult& Ou
 	{
 		OutResult.AnimSetIdx = AnimTrack->AnimSetIdx;
 		OutResult.Animation = Cast<UAnimMontage>(AnimTrack->Animation);
-		OutResult.EntryTransform = AnimTrack->GetAlignmentTransformAtEntryTime() * ToWorldTransform;
-		OutResult.SyncTransform = AnimTrack->GetAlignmentTransformAtSyncTime() * ToWorldTransform;
+		OutResult.EntryTransform = GetAlignmentTransform(*AnimTrack, 0, 0.f) * ToWorldTransform;
+		OutResult.SyncTransform = GetAlignmentTransform(*AnimTrack, 0, AnimTrack->GetSyncTimeForWarpSection(0)) * ToWorldTransform;
 
 		if (QueryParams.bFindAnimStartTime)
 		{
 			const FVector LocalLocation = (Querier.GetTransform().GetRelativeTransform(ToWorldTransform)).GetLocation();
-			OutResult.AnimStartTime = AnimTrack->FindBestAnimStartTime(LocalLocation);
+			OutResult.AnimStartTime = FindBestAnimStartTime(*AnimTrack, LocalLocation);
 		}
 
 		return true;
@@ -665,6 +639,47 @@ bool UContextualAnimSceneAsset::Query(FName Role, FContextualAnimQueryResult& Ou
 	return false;
 }
 
+float UContextualAnimSceneAsset::FindBestAnimStartTime(const FContextualAnimTrack& AnimTrack, const FVector& LocalLocation) const
+{
+	//@TODO: Kept around only to do not break existing content. It will go away in the future.
+
+	float BestTime = 0.f;
+
+	if (AnimTrack.AnimMaxStartTime < 0.f)
+	{
+		return BestTime;
+	}
+
+	const FVector SyncPointLocation = GetAlignmentTransform(AnimTrack, 0, AnimTrack.GetSyncTimeForWarpSection(0)).GetLocation();
+	const float PerfectDistToSyncPointSq = GetAlignmentTransform(AnimTrack, 0, 0.f).GetTranslation().SizeSquared2D();
+	const float ActualDistToSyncPointSq = FVector::DistSquared2D(LocalLocation, SyncPointLocation);
+
+	if (ActualDistToSyncPointSq < PerfectDistToSyncPointSq)
+	{
+		float BestDistance = MAX_FLT;
+		TArrayView<const FVector3f> PosKeys(AnimTrack.AlignmentData.Tracks.AnimationTracks[0].PosKeys.GetData(), AnimTrack.AlignmentData.Tracks.AnimationTracks[0].PosKeys.Num());
+
+		//@TODO: Very simple search for now. Replace with Distance Matching + Pose Matching
+		for (int32 Idx = 0; Idx < PosKeys.Num(); Idx++)
+		{
+			const float Time = Idx * AnimTrack.AlignmentData.SampleInterval;
+			if (AnimTrack.AnimMaxStartTime > 0.f && Time >= AnimTrack.AnimMaxStartTime)
+			{
+				break;
+			}
+
+			const float DistFromCurrentFrameToSyncPointSq = FVector::DistSquared2D(SyncPointLocation, (FVector)PosKeys[Idx]);
+			if (DistFromCurrentFrameToSyncPointSq < ActualDistToSyncPointSq)
+			{
+				BestTime = Time;
+				break;
+			}
+		}
+	}
+
+	return BestTime;
+}
+
 int32 UContextualAnimSceneAsset::GetNumSections() const
 {
 	return Sections.Num();
@@ -680,11 +695,6 @@ const FContextualAnimIKTargetDefContainer& UContextualAnimSceneAsset::GetIKTarge
 	return Sections.IsValidIndex(SectionIdx) ? Sections[SectionIdx].GetIKTargetDefsForRole(Role) : FContextualAnimIKTargetDefContainer::EmptyContainer;
 }
 
-FTransform UContextualAnimSceneAsset::GetAlignmentTransformForRoleRelativeToOtherRoleInSection(int32 SectionIdx, int32 AnimSetIdx, const FName& Role, const FName& OtherRole, float Time) const
-{
-	return Sections.IsValidIndex(SectionIdx) ? Sections[SectionIdx].GetAlignmentTransformForRoleRelativeToOtherRole(AnimSetIdx, Role, OtherRole, Time) : FTransform::Identity;
-}
-
 const TArray<FContextualAnimSetPivotDefinition>& UContextualAnimSceneAsset::GetAnimSetPivotDefinitionsInSection(int32 SectionIdx) const
 {
 	static TArray<FContextualAnimSetPivotDefinition> EmptyDefs;
@@ -701,17 +711,110 @@ FTransform UContextualAnimSceneAsset::GetIKTargetTransform(int32 SectionIdx, int
 	return FTransform::Identity;
 }
 
+FTransform UContextualAnimSceneAsset::GetAlignmentTransform(int32 SectionIdx, int32 AnimSetIdx, int32 AnimTrackIdx, int32 TrackIdx, float Time) const
+{
+	if (const FContextualAnimTrack* AnimTrack = GetAnimTrack(SectionIdx, AnimSetIdx, AnimTrackIdx))
+	{
+		return GetAlignmentTransform(*AnimTrack, TrackIdx, Time);
+	}
+
+	return FTransform::Identity;
+}
+
 FTransform UContextualAnimSceneAsset::GetAlignmentTransform(int32 SectionIdx, int32 AnimSetIdx, int32 AnimTrackIdx, const FName& TrackName, float Time) const
 {
 	if (const FContextualAnimTrack* AnimTrack = GetAnimTrack(SectionIdx, AnimSetIdx, AnimTrackIdx))
 	{
-		return AnimTrack->AlignmentData.ExtractTransformAtTime(TrackName, Time);
+		return GetAlignmentTransform(*AnimTrack, TrackName, Time);
+	}
+
+	return FTransform::Identity;
+}
+
+FTransform UContextualAnimSceneAsset::GetAlignmentTransform(const FContextualAnimTrack& AnimTrack, int32 TrackIdx, float Time) const
+{
+	if (ShouldPrecomputeAlignmentTracks() && AnimTrack.AlignmentData.Tracks.GetNum() > 0)
+	{
+		return AnimTrack.AlignmentData.ExtractTransformAtTime(TrackIdx, Time);
+	}
+
+	const FTransform MeshToComponentInverse = GetMeshToComponentForRole(AnimTrack.Role).Inverse();
+
+	const FContextualAnimSet* AnimSet = GetAnimSet(AnimTrack.SectionIdx, AnimTrack.AnimSetIdx);
+	check(AnimSet);
+
+	FTransform ScenePivotTransform = AnimSet->ScenePivots[TrackIdx];
+
+	if (AnimTrack.Animation)
+	{
+		const FTransform RootTransform = MeshToComponentInverse * (UContextualAnimUtilities::ExtractRootTransformFromAnimation(AnimTrack.Animation, Time) * AnimTrack.MeshToScene);
+		return RootTransform.GetRelativeTransform(ScenePivotTransform);
+	}
+	else
+	{
+		const FTransform RootTransform = MeshToComponentInverse * AnimTrack.MeshToScene;
+		return RootTransform.GetRelativeTransform(ScenePivotTransform);
+	}
+}
+
+FTransform UContextualAnimSceneAsset::GetAlignmentTransform(const FContextualAnimTrack& AnimTrack, const FName& TrackName, float Time) const
+{
+	for (int32 Idx = 0; Idx < Sections[AnimTrack.SectionIdx].AnimSetPivotDefinitions.Num(); Idx++)
+	{
+		if (Sections[AnimTrack.SectionIdx].AnimSetPivotDefinitions[Idx].Name == TrackName)
+		{
+			return GetAlignmentTransform(AnimTrack, Idx, Time);
+		}
 	}
 
 	return FTransform::Identity;
 }
 
-static FContextualAnimPoint GetContextualAnimPoint(const FContextualAnimTrack& PrimaryAnimTrack, const FContextualAnimTrack& SecondaryAnimTrack, const FContextualAnimSceneBindingContext& PrimaryContext, int32 SampleRate, EContextualAnimPointType Type)
+FTransform UContextualAnimSceneAsset::GetAlignmentTransformForRoleRelativeToOtherRole(int32 SectionIdx, int32 AnimSetIdx, FName Role, FName OtherRole, float Time) const
+{
+	if(const FContextualAnimTrack* AnimTrack = GetAnimTrack(SectionIdx, AnimSetIdx, Role))
+	{
+		if (const FContextualAnimTrack* OtherAnimTrack = GetAnimTrack(SectionIdx, AnimSetIdx, OtherRole))
+		{
+			const FTransform FromRoleRelativeToScenePivot = GetAlignmentTransform(SectionIdx, AnimSetIdx, AnimTrack->AnimTrackIdx, 0, Time);//GetAlignmentTransformForRoleRelativeToPivot(AnimSetIdx, Role, Time);
+			const FTransform ToRoleRelativeToScenePivot = GetAlignmentTransform(SectionIdx, AnimSetIdx, OtherAnimTrack->AnimTrackIdx, 0, Time); //GetAlignmentTransformForRoleRelativeToPivot(AnimSetIdx, OtherRole, Time);
+			return FromRoleRelativeToScenePivot.GetRelativeTransform(ToRoleRelativeToScenePivot);
+		}
+	}
+	
+	return FTransform::Identity;
+}
+
+const FContextualAnimTrack* UContextualAnimSceneAsset::FindAnimTrackForRoleWithClosestEntryLocation(int32 SectionIdx, const FName& Role, const FContextualAnimSceneBindingContext& Primary, const FVector& TestLocation) const
+{
+	const FContextualAnimTrack* Result = nullptr;
+
+	float BestDistanceSq = MAX_FLT;
+
+	if(const FContextualAnimSceneSection* Section = GetSection(SectionIdx))
+	{
+		for (const FContextualAnimSet& AnimSet : Section->AnimSets)
+		{
+			for (const FContextualAnimTrack& AnimTrack : AnimSet.Tracks)
+			{
+				if (AnimTrack.Role == Role)
+				{
+					const FTransform EntryTransform = GetAlignmentTransform(AnimTrack, 0, 0.f) * Primary.GetTransform();
+					const float DistSq = FVector::DistSquared(EntryTransform.GetLocation(), TestLocation);
+					if (DistSq < BestDistanceSq)
+					{
+						BestDistanceSq = DistSq;
+						Result = &AnimTrack;
+					}
+				}
+			}
+		}
+	}
+
+	return Result;
+}
+
+static FContextualAnimPoint GetContextualAnimPoint(const UContextualAnimSceneAsset& SceneAsset, const FContextualAnimTrack& PrimaryAnimTrack, const FContextualAnimTrack& SecondaryAnimTrack, const FContextualAnimSceneBindingContext& PrimaryContext, int32 SampleRate, EContextualAnimPointType Type)
 {
 	check(SampleRate > 0);
 
@@ -734,8 +837,8 @@ static FContextualAnimPoint GetContextualAnimPoint(const FContextualAnimTrack& P
 		const float Delta = (SecondaryAnimTrack.GetRootTransformAtTime(T2).GetTranslation() - SecondaryAnimTrack.GetRootTransformAtTime(T1).GetTranslation()).Size();
 		const float Speed = Delta / Interval;
 
-		const FTransform SecondaryRelativeToPivot = SecondaryAnimTrack.GetAlignmentTransformAtTime(T1);
-		const FTransform PrimaryRelativeToPivot = PrimaryAnimTrack.GetAlignmentTransformAtTime(T1);
+		const FTransform SecondaryRelativeToPivot = SceneAsset.GetAlignmentTransform(SecondaryAnimTrack.SectionIdx, SecondaryAnimTrack.AnimSetIdx, SecondaryAnimTrack.AnimTrackIdx, 0, T1);
+		const FTransform PrimaryRelativeToPivot = SceneAsset.GetAlignmentTransform(PrimaryAnimTrack.SectionIdx, PrimaryAnimTrack.AnimSetIdx, PrimaryAnimTrack.AnimTrackIdx, 0, T1);
 		const FTransform FinalTransform = SecondaryRelativeToPivot.GetRelativeTransform(PrimaryRelativeToPivot) * PrimaryContext.GetTransform();
 
 		return FContextualAnimPoint(SecondaryAnimTrack.Role, FinalTransform, Speed, SecondaryAnimTrack.SectionIdx, SecondaryAnimTrack.AnimSetIdx, SecondaryAnimTrack.AnimTrackIdx);
@@ -759,7 +862,7 @@ void UContextualAnimSceneAsset::GetAlignmentPointsForSecondaryRole(EContextualAn
 			
 			if(PrimaryAnimTrack && SecondaryAnimTrack)
 			{
-				OutResult.Add(GetContextualAnimPoint(*PrimaryAnimTrack, *SecondaryAnimTrack, Primary, GetSampleRate(), Type));
+				OutResult.Add(GetContextualAnimPoint(*this, *PrimaryAnimTrack, *SecondaryAnimTrack, Primary, GetSampleRate(), Type));
 			}
 		}
 	}
@@ -798,7 +901,7 @@ void UContextualAnimSceneAsset::GetAlignmentPointsForSecondaryRoleConsideringSel
 
 				if (bSuccess)
 				{
-					OutResult.Add(GetContextualAnimPoint(*PrimaryAnimTrack, *SecondaryAnimTrack, Primary, GetSampleRate(), Type));
+					OutResult.Add(GetContextualAnimPoint(*this, *PrimaryAnimTrack, *SecondaryAnimTrack, Primary, GetSampleRate(), Type));
 				}
 			}
 		}
@@ -816,7 +919,12 @@ UAnimSequenceBase* UContextualAnimSceneAsset::BP_FindAnimationForRole(int32 Sect
 
 FTransform UContextualAnimSceneAsset::BP_GetAlignmentTransformForRoleRelativeToPivot(int32 SectionIdx, int32 AnimSetIdx, FName Role, float Time) const
 {
-	return (Sections.IsValidIndex(SectionIdx)) ? Sections[SectionIdx].GetAlignmentTransformForRoleRelativeToPivot(AnimSetIdx, Role, Time) : FTransform::Identity;
+	if (const FContextualAnimTrack* AnimTrack = GetAnimTrack(SectionIdx, AnimSetIdx, Role))
+	{
+		return GetAlignmentTransform(AnimTrack->SectionIdx, AnimTrack->AnimSetIdx, AnimTrack->AnimTrackIdx, 0, Time);
+	}
+
+	return FTransform::Identity;
 }
 
 FTransform UContextualAnimSceneAsset::BP_GetIKTargetTransformForRoleAtTime(int32 SectionIdx, int32 AnimSetIdx, FName Role, FName TrackName, float Time) const
diff --git a/Engine/Plugins/Experimental/Animation/ContextualAnimation/Source/ContextualAnimation/Private/ContextualAnimTypes.cpp b/Engine/Plugins/Experimental/Animation/ContextualAnimation/Source/ContextualAnimation/Private/ContextualAnimTypes.cpp
--- a/Engine/Plugins/Experimental/Animation/ContextualAnimation/Source/ContextualAnimation/Private/ContextualAnimTypes.cpp
+++ b/Engine/Plugins/Experimental/Animation/ContextualAnimation/Source/ContextualAnimation/Private/ContextualAnimTypes.cpp
@@ -164,45 +164,6 @@ void FContextualAnimTrack::GetStartAndEndTimeForWarpSection(const FName& WarpSec
 	}
 }
 
-float FContextualAnimTrack::FindBestAnimStartTime(const FVector& LocalLocation) const
-{
-	float BestTime = 0.f;
-
-	if (AnimMaxStartTime < 0.f)
-	{
-		return BestTime;
-	}
-
-	const FVector SyncPointLocation = GetAlignmentTransformAtSyncTime().GetLocation();
-	const float PerfectDistToSyncPointSq = GetAlignmentTransformAtEntryTime().GetTranslation().SizeSquared2D();
-	const float ActualDistToSyncPointSq = FVector::DistSquared2D(LocalLocation, SyncPointLocation);
-
-	if (ActualDistToSyncPointSq < PerfectDistToSyncPointSq)
-	{
-		float BestDistance = MAX_FLT;
-		TArrayView<const FVector3f> PosKeys(AlignmentData.Tracks.AnimationTracks[0].PosKeys.GetData(), AlignmentData.Tracks.AnimationTracks[0].PosKeys.Num());
-
-		//@TODO: Very simple search for now. Replace with Distance Matching + Pose Matching
-		for (int32 Idx = 0; Idx < PosKeys.Num(); Idx++)
-		{
-			const float Time = Idx * AlignmentData.SampleInterval;
-			if (AnimMaxStartTime > 0.f && Time >= AnimMaxStartTime)
-			{
-				break;
-			}
-
-			const float DistFromCurrentFrameToSyncPointSq = FVector::DistSquared2D(SyncPointLocation, (FVector)PosKeys[Idx]);
-			if (DistFromCurrentFrameToSyncPointSq < ActualDistToSyncPointSq)
-			{
-				BestTime = Time;
-				break;
-			}
-		}
-	}
-
-	return BestTime;
-}
-
 bool FContextualAnimTrack::DoesQuerierPassSelectionCriteria(const FContextualAnimSceneBindingContext& Primary, const FContextualAnimSceneBindingContext& Querier) const
 {
 	for (const UContextualAnimSelectionCriterion* Criterion : SelectionCriteria)
diff --git a/Engine/Plugins/Experimental/Animation/ContextualAnimation/Source/ContextualAnimation/Private/ContextualAnimUtilities.cpp b/Engine/Plugins/Experimental/Animation/ContextualAnimation/Source/ContextualAnimation/Private/ContextualAnimUtilities.cpp
--- a/Engine/Plugins/Experimental/Animation/ContextualAnimation/Source/ContextualAnimation/Private/ContextualAnimUtilities.cpp
+++ b/Engine/Plugins/Experimental/Animation/ContextualAnimation/Source/ContextualAnimation/Private/ContextualAnimUtilities.cpp
@@ -164,9 +164,9 @@ void UContextualAnimUtilities::DrawDebugAnimSet(const UWorld* World, const UCont
 {
 	if (World)
 	{
-		for(const FContextualAnimTrack& AnimTrack : AnimSet.Tracks)
+		for (const FContextualAnimTrack& AnimTrack : AnimSet.Tracks)
 		{
-			const FTransform Transform = (SceneAsset.GetMeshToComponentForRole(AnimTrack.Role) * AnimTrack.GetAlignmentTransformAtTime(Time)) * ToWorldTransform;
+			const FTransform Transform = (SceneAsset.GetMeshToComponentForRole(AnimTrack.Role) * SceneAsset.GetAlignmentTransform(AnimTrack, 0, Time)) * ToWorldTransform;
 
 			if (const UAnimSequenceBase* Animation = AnimTrack.Animation)
 			{
@@ -449,7 +449,7 @@ FTransform UContextualAnimUtilities::BP_SceneBindings_GetAlignmentTransformForRo
 
 	if(const UContextualAnimSceneAsset* SceneAsset = Bindings.GetSceneAsset())
 	{
-		Result = SceneAsset->GetAlignmentTransformForRoleRelativeToOtherRoleInSection(Bindings.GetSectionIdx(), Bindings.GetAnimSetIdx(), Role, RelativeToRole, Time);
+		Result = SceneAsset->GetAlignmentTransformForRoleRelativeToOtherRole(Bindings.GetSectionIdx(), Bindings.GetAnimSetIdx(), Role, RelativeToRole, Time);
 	}
 
 	return Result;
@@ -461,9 +461,10 @@ FTransform UContextualAnimUtilities::BP_SceneBindings_GetAlignmentTransformForRo
 
 	if (const UContextualAnimSceneAsset* SceneAsset = Bindings.GetSceneAsset())
 	{
-		if (const FContextualAnimSceneSection* Section = SceneAsset->GetSection(Bindings.GetSectionIdx()))
+		if(const FContextualAnimSceneBinding* Binding = Bindings.FindBindingByRole(Role))
 		{
-			return Section->GetAlignmentTransformForRoleRelativeToPivot(Bindings.GetAnimSetIdx(), Role, Time) * Pivot.Transform;
+			const FContextualAnimTrack& AnimTrack = Bindings.GetAnimTrackFromBinding(*Binding);
+			return SceneAsset->GetAlignmentTransform(AnimTrack, 0, Time);
 		}
 	}
 
@@ -505,5 +506,5 @@ FTransform UContextualAnimUtilities::BP_SceneBindings_GetAlignmentTransformFromB
 	float StartTime, EndTime;
 	AnimTrack.GetStartAndEndTimeForWarpSection(Pivot.Name, StartTime, EndTime);
 
-	return AnimTrack.GetAlignmentTransformAtTime(EndTime) * Pivot.Transform;
+	return Bindings.GetSceneAsset()->GetAlignmentTransform(AnimTrack, 0, EndTime) * Pivot.Transform;
 }
diff --git a/Engine/Plugins/Experimental/Animation/ContextualAnimation/Source/ContextualAnimation/Public/ContextualAnimSceneAsset.h b/Engine/Plugins/Experimental/Animation/ContextualAnimation/Source/ContextualAnimation/Public/ContextualAnimSceneAsset.h
--- a/Engine/Plugins/Experimental/Animation/ContextualAnimation/Source/ContextualAnimation/Public/ContextualAnimSceneAsset.h
+++ b/Engine/Plugins/Experimental/Animation/ContextualAnimation/Source/ContextualAnimation/Public/ContextualAnimSceneAsset.h
@@ -83,18 +83,12 @@ public:
 
 	const FContextualAnimTrack* GetAnimTrack(int32 AnimSetIdx, int32 AnimTrackIdx) const;
 
-	FTransform GetAlignmentTransformForRoleRelativeToPivot(int32 AnimSetIdx, FName Role, float Time) const;
-
-	FTransform GetAlignmentTransformForRoleRelativeToOtherRole(int32 AnimSetIdx, FName Role, FName OtherRole, float Time) const;
-
 	FTransform GetIKTargetTransformForRoleAtTime(int32 AnimSetIdx, FName Role, FName TrackName, float Time) const;
 
 	const FContextualAnimIKTargetDefContainer& GetIKTargetDefsForRole(const FName& Role) const;
 
 	const FContextualAnimTrack* FindFirstAnimTrackForRoleThatPassesSelectionCriteria(const FName& Role, const FContextualAnimSceneBindingContext& Primary, const FContextualAnimSceneBindingContext& Querier) const;
 
-	const FContextualAnimTrack* FindAnimTrackForRoleWithClosestEntryLocation(const FName& Role, const FContextualAnimSceneBindingContext& Primary, const FVector& TestLocation) const;
-
 	FORCEINLINE FName GetName() const { return Name; }
 	FORCEINLINE const TArray<FContextualAnimSetPivotDefinition>& GetAnimSetPivotDefinitions() const { return AnimSetPivotDefinitions; }
 	FORCEINLINE int32 GetNumAnimSets() const { return AnimSets.Num(); }
@@ -223,6 +217,7 @@ public:
 	FORCEINLINE float GetRadius() const { return Radius; }
 	FORCEINLINE const TArray<FContextualAnimWarpPointData>& GetWarpPoints() const { return WarpPoints; }
 	FORCEINLINE const FTransform* GetWarpPointTransform(FName SocketName) const { return WarpPointTransformCache.Find(SocketName); }
+	FORCEINLINE bool ShouldPrecomputeAlignmentTracks() const { return bPrecomputeAlignmentTracks; }
 
 	bool HasValidData() const { return RolesAsset != nullptr && Sections.Num() > 0 && Sections[0].AnimSets.Num() > 0; }
 
@@ -257,14 +252,18 @@ public:
 
 	FTransform GetIKTargetTransform(int32 SectionIdx, int32 AnimSetIdx, int32 AnimTrackIdx, const FName& TrackName, float Time) const;
 
+	FTransform GetAlignmentTransform(int32 SectionIdx, int32 AnimSetIdx, int32 AnimTrackIdx, int32 TrackIdx, float Time) const;
 	FTransform GetAlignmentTransform(int32 SectionIdx, int32 AnimSetIdx, int32 AnimTrackIdx, const FName& TrackName, float Time) const;
+	FTransform GetAlignmentTransform(const FContextualAnimTrack& AnimTrack, int32 TrackIdx, float Time) const;
+	FTransform GetAlignmentTransform(const FContextualAnimTrack& AnimTrack, const FName& TrackName, float Time) const;
+	FTransform GetAlignmentTransformForRoleRelativeToOtherRole(int32 SectionIdx, int32 AnimSetIdx, FName Role, FName OtherRole, float Time) const;
+
+	const FContextualAnimTrack* FindAnimTrackForRoleWithClosestEntryLocation(int32 SectionIdx, const FName& Role, const FContextualAnimSceneBindingContext& Primary, const FVector& TestLocation) const;
 
 	const FContextualAnimTrack* FindAnimTrackByAnimation(const UAnimSequenceBase* Animation) const;
 
 	const TArray<FContextualAnimSetPivotDefinition>& GetAnimSetPivotDefinitionsInSection(int32 SectionIdx) const;
 
-	FTransform GetAlignmentTransformForRoleRelativeToOtherRoleInSection(int32 SectionIdx, int32 AnimSetIdx, const FName& Role, const FName& OtherRole, float Time) const;
-
 	const FContextualAnimIKTargetDefContainer& GetIKTargetDefsForRoleInSection(int32 SectionIdx, const FName& Role) const;
 
 	UFUNCTION(BlueprintCallable, Category = "Contextual Anim|Scene Asset")
@@ -296,6 +295,7 @@ public:
 	//@TODO: Kept around only to do not break existing content. It will go away in the future.
 	UFUNCTION(BlueprintCallable, Category = "Contextual Anim|Scene Asset")
 	bool Query(FName Role, FContextualAnimQueryResult& OutResult, const FContextualAnimQueryParams& QueryParams, const FTransform& ToWorldTransform) const;
+	float FindBestAnimStartTime(const FContextualAnimTrack& AnimTrack, const FVector& LocalLocation) const;
 
 protected:
 
@@ -335,6 +335,10 @@ protected:
 	UPROPERTY(EditAnywhere, Category = "Settings")
 	bool bDisableCollisionBetweenActors = true;
 
+	/** Whether we should extract and cache alignment tracks off line. */
+	UPROPERTY(EditAnywhere, Category = "Settings", AdvancedDisplay)
+	bool bPrecomputeAlignmentTracks = false;
+
 	/** Sample rate (frames per second) used when sampling the animations to generate alignment and IK tracks */
 	UPROPERTY(EditAnywhere, Category = "Settings", meta = (ClampMin = "1", ClampMax = "60"), AdvancedDisplay)
 	int32 SampleRate = 15;
diff --git a/Engine/Plugins/Experimental/Animation/ContextualAnimation/Source/ContextualAnimation/Public/ContextualAnimTypes.h b/Engine/Plugins/Experimental/Animation/ContextualAnimation/Source/ContextualAnimation/Public/ContextualAnimTypes.h
--- a/Engine/Plugins/Experimental/Animation/ContextualAnimation/Source/ContextualAnimation/Public/ContextualAnimTypes.h
+++ b/Engine/Plugins/Experimental/Animation/ContextualAnimation/Source/ContextualAnimation/Public/ContextualAnimTypes.h
@@ -112,12 +112,6 @@ struct CONTEXTUALANIMATION_API FContextualAnimTrack
 	void GetStartAndEndTimeForWarpSection(int32 WarpSectionIndex, float& OutStartTime, float& OutEndTime) const;
 	void GetStartAndEndTimeForWarpSection(const FName& WarpSectionName, float& OutStartTime, float& OutEndTime) const;
 
-	FORCEINLINE FTransform GetAlignmentTransformAtTime(float Time) const { return AlignmentData.ExtractTransformAtTime(0, Time); }
-	FORCEINLINE FTransform GetAlignmentTransformAtEntryTime() const { return AlignmentData.ExtractTransformAtTime(0, 0.f); }
-	FORCEINLINE FTransform GetAlignmentTransformAtSyncTime() const { return AlignmentData.ExtractTransformAtTime(0, GetSyncTimeForWarpSection(0)); }
-
-	float FindBestAnimStartTime(const FVector& LocalLocation) const;
-
 	bool DoesQuerierPassSelectionCriteria(const FContextualAnimSceneBindingContext& PrimaryActorData, const FContextualAnimSceneBindingContext& QuerierData) const;
 
 	FTransform GetRootTransformAtTime(float Time) const;
diff --git a/Engine/Plugins/Experimental/Animation/ContextualAnimation/Source/ContextualAnimationEditor/Private/ContextualAnimEdMode.cpp b/Engine/Plugins/Experimental/Animation/ContextualAnimation/Source/ContextualAnimationEditor/Private/ContextualAnimEdMode.cpp
--- a/Engine/Plugins/Experimental/Animation/ContextualAnimation/Source/ContextualAnimationEditor/Private/ContextualAnimEdMode.cpp
+++ b/Engine/Plugins/Experimental/Animation/ContextualAnimation/Source/ContextualAnimationEditor/Private/ContextualAnimEdMode.cpp
@@ -199,7 +199,7 @@ void FContextualAnimEdMode::Render(const FSceneView* View, FViewport* Viewport,
 											FVector Origin, Direction;
 											if (Cone->Mode == EContextualAnimCriterionConeMode::ToPrimary)
 											{
-												const FTransform QuerierTransform = Section->GetAlignmentTransformForRoleRelativeToOtherRole(AnimSetIdx, AnimTrack.Role, PrimaryRole, 0.f) * PrimaryTransform;
+												const FTransform QuerierTransform = SceneAsset->GetAlignmentTransformForRoleRelativeToOtherRole(AnimTrack.SectionIdx, AnimTrack.AnimSetIdx, AnimTrack.Role, PrimaryRole, 0.f) * PrimaryTransform;
 												Origin = QuerierTransform.GetLocation();
 												Direction = QuerierTransform.GetRotation().GetForwardVector().RotateAngleAxis(Cone->Offset, FVector::UpVector);
 											}
@@ -222,7 +222,7 @@ void FContextualAnimEdMode::Render(const FSceneView* View, FViewport* Viewport,
 								if (AnimTrack.Role != PrimaryRole && AnimTrack.Animation)
 								{
 									const float Time = 0.f;
-									const FTransform RelativeToPrimary = Section->GetAlignmentTransformForRoleRelativeToOtherRole(AnimSetIdx, AnimTrack.Role, PrimaryRole, Time);
+									const FTransform RelativeToPrimary = SceneAsset->GetAlignmentTransformForRoleRelativeToOtherRole(AnimTrack.SectionIdx, AnimTrack.AnimSetIdx, AnimTrack.Role, PrimaryRole, Time);
 									FTransform Transform = (SceneAsset->GetMeshToComponentForRole(AnimTrack.Role) * RelativeToPrimary) * PrimaryTransform;
 
 									if (ACharacter* PrimaryAsCharacter = Cast<ACharacter>(PrimaryBinding->GetActor()))
