Contextual Anim: Replaced assert with log and early out to prevent random crash when calculating warp points. The issue needs further investigation.

diff --git a/Engine/Plugins/Experimental/Animation/ContextualAnimation/Source/ContextualAnimation/Private/ContextualAnimTypes.cpp b/Engine/Plugins/Experimental/Animation/ContextualAnimation/Source/ContextualAnimation/Private/ContextualAnimTypes.cpp
--- a/Engine/Plugins/Experimental/Animation/ContextualAnimation/Source/ContextualAnimation/Private/ContextualAnimTypes.cpp
+++ b/Engine/Plugins/Experimental/Animation/ContextualAnimation/Source/ContextualAnimation/Private/ContextualAnimTypes.cpp
@@ -641,7 +641,14 @@ void FContextualAnimSceneBindings::CalculateWarpPoints(TArray<FContextualAnimWar
 
 bool FContextualAnimSceneBindings::CalculateWarpPoint(const FContextualAnimWarpPointDefinition& WarpPointDef, FContextualAnimWarpPoint& OutWarpPoint) const
 {
-	check(IsValid());
+	if (!IsValid())
+	{
+		const UContextualAnimSceneAsset* Asset = SceneAsset.Get();
+		UE_LOG(LogContextualAnim, Warning, TEXT("FContextualAnimSceneBindings::CalculateWarpPoint failed. Reason: Invalid Bindings. Id: %d Asset: %s HasValidData: %d Num: %d"), 
+			GetID(), *GetNameSafe(Asset), Asset ? Asset->HasValidData() : 0, Num());
+
+		return false;
+	}
 
 	if (WarpPointDef.Mode == EContextualAnimWarpPointDefinitionMode::PrimaryActor)
 	{
