Contextual Anim: Allow to set warp targets manually

diff --git a/Engine/Plugins/Experimental/Animation/ContextualAnimation/Source/ContextualAnimation/Private/ContextualAnimSceneActorComponent.cpp b/Engine/Plugins/Experimental/Animation/ContextualAnimation/Source/ContextualAnimation/Private/ContextualAnimSceneActorComponent.cpp
--- a/Engine/Plugins/Experimental/Animation/ContextualAnimation/Source/ContextualAnimation/Private/ContextualAnimSceneActorComponent.cpp
+++ b/Engine/Plugins/Experimental/Animation/ContextualAnimation/Source/ContextualAnimation/Private/ContextualAnimSceneActorComponent.cpp
@@ -149,10 +149,15 @@ void UContextualAnimSceneActorComponent::AddOrUpdateWarpTargets(int32 SectionIdx
 				}
 			}
 		}
+
+		if (Binding->HasExternalWarpTarget())
+		{
+			MotionWarpComp->AddOrUpdateWarpTargetFromTransform(Binding->GetExternalWarpTargetName(), Binding->GetExternalWarpTargetTransform());
+		}
 	}
 }
 
-bool UContextualAnimSceneActorComponent::LateJoinContextualAnimScene(AActor* Actor, FName Role)
+bool UContextualAnimSceneActorComponent::LateJoinContextualAnimScene(AActor* Actor, FName Role, const TArray<FContextualAnimWarpTarget>& WarpTargets)
 {
 	if (!Bindings.IsValid())
 	{
@@ -168,7 +173,7 @@ bool UContextualAnimSceneActorComponent::LateJoinContextualAnimScene(AActor* Act
 		{
 			if (UContextualAnimSceneActorComponent* Comp = Leader->GetSceneActorComponent())
 			{
-				return Comp->LateJoinContextualAnimScene(Actor, Role);
+				return Comp->LateJoinContextualAnimScene(Actor, Role, WarpTargets);
 			}
 		}
 	}
@@ -177,7 +182,7 @@ bool UContextualAnimSceneActorComponent::LateJoinContextualAnimScene(AActor* Act
 		*UEnum::GetValueAsString(TEXT("Engine.ENetRole"), GetOwner()->GetLocalRole()), *GetNameSafe(GetOwner()), Bindings.GetID(), Bindings.GetSectionIdx(), *GetNameSafe(Bindings.GetSceneAsset()), *GetNameSafe(Actor), *Role.ToString());
 
 	// Play animation and set state on this new actor that is joining us and update bindings for everyone else
-	if (HandleLateJoin(Actor, Role))
+	if (HandleLateJoin(Actor, Role, WarpTargets))
 	{
 		// Replicate late join event. See OnRep_LateJoinData
 		if (GetOwner()->HasAuthority())
@@ -195,7 +200,7 @@ bool UContextualAnimSceneActorComponent::LateJoinContextualAnimScene(AActor* Act
 	return false;
 }
 
-bool UContextualAnimSceneActorComponent::HandleLateJoin(AActor* Actor, FName Role)
+bool UContextualAnimSceneActorComponent::HandleLateJoin(AActor* Actor, FName Role, const TArray<FContextualAnimWarpTarget>& WarpTargets)
 {
 	if (!Bindings.BindActorToRole(*Actor, Role))
 	{
@@ -222,15 +227,14 @@ bool UContextualAnimSceneActorComponent::HandleLateJoin(AActor* Actor, FName Rol
 	{
 		if (UContextualAnimSceneActorComponent* Comp = Binding->GetSceneActorComponent())
 		{
-			Comp->LateJoinScene(Bindings);
+			Comp->LateJoinScene(Bindings, WarpTargets);
 		}
 	}
-	
 
 	return true;
 }
 
-void UContextualAnimSceneActorComponent::LateJoinScene(const FContextualAnimSceneBindings& InBindings)
+void UContextualAnimSceneActorComponent::LateJoinScene(const FContextualAnimSceneBindings& InBindings, const TArray<FContextualAnimWarpTarget>& WarpTargets)
 {
 	if (Bindings.IsValid())
 	{
@@ -247,6 +251,11 @@ void UContextualAnimSceneActorComponent::LateJoinScene(const FContextualAnimScen
 
 		Bindings = InBindings;
 
+		for (const FContextualAnimWarpTarget& WarpTarget : WarpTargets)
+		{
+			Bindings.SetRoleWarpTarget(WarpTarget.Role, WarpTarget.TargetName, WarpTarget.TargetTransform);
+		}
+
 		// For now when late joining an scene always play animation from first section
 		const int32 SectionIdx = 0;
 		const int32 AnimSetIdx = 0;
@@ -277,10 +286,10 @@ void UContextualAnimSceneActorComponent::OnRep_LateJoinData()
 		*UEnum::GetValueAsString(TEXT("Engine.ENetRole"), GetOwner()->GetLocalRole()), *GetNameSafe(GetOwner()), Bindings.GetID(), Bindings.GetSectionIdx(), *GetNameSafe(Bindings.GetSceneAsset()), *GetNameSafe(RepLateJoinData.Actor), *RepLateJoinData.Role.ToString());
 
 	// Play animation and set state on this new actor that is joining us and update bindings for everyone else
-	HandleLateJoin(RepLateJoinData.Actor, RepLateJoinData.Role);
+	HandleLateJoin(RepLateJoinData.Actor, RepLateJoinData.Role, {});
 }
 
-bool UContextualAnimSceneActorComponent::TransitionContextualAnimScene(FName SectionName)
+bool UContextualAnimSceneActorComponent::TransitionContextualAnimScene(FName SectionName, const TArray<FContextualAnimWarpTarget>& WarpTargets)
 {
 	if (!GetOwner()->HasAuthority())
 	{
@@ -295,13 +304,18 @@ bool UContextualAnimSceneActorComponent::TransitionContextualAnimScene(FName Sec
 		{
 			if (UContextualAnimSceneActorComponent* Comp = Leader->GetSceneActorComponent())
 			{
-				return Comp->TransitionContextualAnimScene(SectionName);
+				return Comp->TransitionContextualAnimScene(SectionName, WarpTargets);
 			}
 		}
 	}
 
 	if (const FContextualAnimSceneBinding* OwnerBinding = Bindings.FindBindingByActor(GetOwner()))
 	{
+		for (const FContextualAnimWarpTarget& WarpTarget : WarpTargets)
+		{
+			Bindings.SetRoleWarpTarget(WarpTarget.Role, WarpTarget.TargetName, WarpTarget.TargetTransform);
+		}
+
 		const int32 SectionIdx = Bindings.GetSceneAsset()->GetSectionIndex(SectionName);
 		if (SectionIdx != INDEX_NONE)
 		{
@@ -355,7 +369,7 @@ void UContextualAnimSceneActorComponent::HandleTransitionSelf(int32 NewSectionId
 	AddOrUpdateWarpTargets(NewSectionIdx, NewAnimSetIdx);
 }
 
-bool UContextualAnimSceneActorComponent::TransitionSingleActor(int32 SectionIdx, int32 AnimSetIdx)
+bool UContextualAnimSceneActorComponent::TransitionSingleActor(int32 SectionIdx, int32 AnimSetIdx, const TArray<FContextualAnimWarpTarget>& WarpTargets)
 {
 	if (!GetOwner()->HasAuthority())
 	{
@@ -364,6 +378,11 @@ bool UContextualAnimSceneActorComponent::TransitionSingleActor(int32 SectionIdx,
 
 	if (const FContextualAnimSceneBinding* OwnerBinding = Bindings.FindBindingByActor(GetOwner()))
 	{
+		for (const FContextualAnimWarpTarget& WarpTarget : WarpTargets)
+		{
+			Bindings.SetRoleWarpTarget(WarpTarget.Role, WarpTarget.TargetName, WarpTarget.TargetTransform);
+		}
+
 		if (const UContextualAnimSceneAsset* Asset = Bindings.GetSceneAsset())
 		{
 			const FContextualAnimTrack* AnimTrack = Asset->GetAnimTrack(SectionIdx, AnimSetIdx, Bindings.GetRoleFromBinding(*OwnerBinding));
@@ -410,6 +429,26 @@ void UContextualAnimSceneActorComponent::OnRep_RepTransitionSingleActor()
 }
 
 bool UContextualAnimSceneActorComponent::StartContextualAnimScene(const FContextualAnimSceneBindings& InBindings)
+{
+	return StartContextualAnimScene(InBindings, {});
+}
+	
+bool UContextualAnimSceneActorComponent::LateJoinContextualAnimScene(AActor* Actor, FName Role)
+{
+	return LateJoinContextualAnimScene(Actor, Role, {});
+}
+	
+bool UContextualAnimSceneActorComponent::TransitionContextualAnimScene(FName SectionName)
+{
+	return TransitionContextualAnimScene(SectionName, {});
+}
+	
+bool UContextualAnimSceneActorComponent::TransitionSingleActor(int32 SectionIdx, int32 AnimSetIdx)
+{
+	return TransitionSingleActor(SectionIdx, AnimSetIdx, {});
+}
+
+bool UContextualAnimSceneActorComponent::StartContextualAnimScene(const FContextualAnimSceneBindings& InBindings, const TArray<FContextualAnimWarpTarget>& WarpTargets)
 {
 	UE_LOG(LogContextualAnim, Log, TEXT("%-21s UContextualAnimSceneActorComponent::StartContextualAnim Actor: %s"),
 		*UEnum::GetValueAsString(TEXT("Engine.ENetRole"), GetOwner()->GetLocalRole()), *GetNameSafe(GetOwner()));
@@ -419,7 +458,7 @@ bool UContextualAnimSceneActorComponent::StartContextualAnimScene(const FContext
 	{
 		if (GetOwner()->HasAuthority())
 		{
-			JoinScene(InBindings);
+			JoinScene(InBindings, WarpTargets);
 
 			for (const FContextualAnimSceneBinding& Binding : InBindings)
 			{
@@ -427,7 +466,7 @@ bool UContextualAnimSceneActorComponent::StartContextualAnimScene(const FContext
 				{
 					if (UContextualAnimSceneActorComponent* Comp = Binding.GetSceneActorComponent())
 					{
-						Comp->JoinScene(InBindings);
+						Comp->JoinScene(InBindings, WarpTargets);
 					}
 				}
 			}
@@ -440,7 +479,7 @@ bool UContextualAnimSceneActorComponent::StartContextualAnimScene(const FContext
 		}
 		else if (GetOwner()->GetLocalRole() == ROLE_AutonomousProxy)
 		{
-			JoinScene(InBindings);
+			JoinScene(InBindings, WarpTargets);
 
 			ServerStartContextualAnimScene(InBindings);
 
@@ -453,7 +492,7 @@ bool UContextualAnimSceneActorComponent::StartContextualAnimScene(const FContext
 
 void UContextualAnimSceneActorComponent::ServerStartContextualAnimScene_Implementation(const FContextualAnimSceneBindings& InBindings)
 {
-	StartContextualAnimScene(InBindings);
+	StartContextualAnimScene(InBindings, {});
 }
 
 bool UContextualAnimSceneActorComponent::ServerStartContextualAnimScene_Validate(const FContextualAnimSceneBindings& InBindings)
@@ -530,7 +569,7 @@ void UContextualAnimSceneActorComponent::OnRep_Bindings()
 			// Join the scene (start playing animation, etc.)
 			if (GetOwner()->GetLocalRole() != ROLE_AutonomousProxy)
 			{
-				JoinScene(RepBindings);
+				JoinScene(RepBindings, {});
 			}
 
 			// RepBindings is only replicated from the initiator of the action.
@@ -543,7 +582,7 @@ void UContextualAnimSceneActorComponent::OnRep_Bindings()
 				{
 					if (UContextualAnimSceneActorComponent* Comp = Binding.GetSceneActorComponent())
 					{
-						Comp->JoinScene(RepBindings);
+						Comp->JoinScene(RepBindings, {});
 					}
 				}
 			}
@@ -681,7 +720,7 @@ void UContextualAnimSceneActorComponent::OnLeftScene()
 	}
 }
 
-void UContextualAnimSceneActorComponent::JoinScene(const FContextualAnimSceneBindings& InBindings)
+void UContextualAnimSceneActorComponent::JoinScene(const FContextualAnimSceneBindings& InBindings, const TArray<FContextualAnimWarpTarget>& WarpTargets)
 {
 	if (Bindings.IsValid())
 	{
@@ -695,6 +734,11 @@ void UContextualAnimSceneActorComponent::JoinScene(const FContextualAnimSceneBin
 
 		Bindings = InBindings;
 
+		for (const FContextualAnimWarpTarget& WarpTarget : WarpTargets)
+		{
+			Bindings.SetRoleWarpTarget(WarpTarget.Role, WarpTarget.TargetName, WarpTarget.TargetTransform);
+		}
+
 		const FContextualAnimTrack& AnimTrack = Bindings.GetAnimTrackFromBinding(*Binding);
 		PlayAnimation_Internal(AnimTrack.Animation, 0.f, true);
 
diff --git a/Engine/Plugins/Experimental/Animation/ContextualAnimation/Source/ContextualAnimation/Private/ContextualAnimTypes.cpp b/Engine/Plugins/Experimental/Animation/ContextualAnimation/Source/ContextualAnimation/Private/ContextualAnimTypes.cpp
--- a/Engine/Plugins/Experimental/Animation/ContextualAnimation/Source/ContextualAnimation/Private/ContextualAnimTypes.cpp
+++ b/Engine/Plugins/Experimental/Animation/ContextualAnimation/Source/ContextualAnimation/Private/ContextualAnimTypes.cpp
@@ -274,6 +274,11 @@ void FContextualAnimSceneBinding::SetAnimTrack(const FContextualAnimTrack& InAni
 	AnimTrackIdx = InAnimTrack.AnimTrackIdx;
 }
 
+void FContextualAnimSceneBinding::SetExternalWarpTargetTransform(const FName TargetName, const FTransform& Transform)
+{
+	Context.SetExternalWarpTargetTransform(TargetName, Transform);
+}
+
 FAnimMontageInstance* FContextualAnimSceneBinding::GetAnimMontageInstance() const
 {
 	if (UAnimInstance* AnimInstance = GetAnimInstance())
@@ -386,6 +391,16 @@ bool FContextualAnimSceneBindings::BindActorToRole(AActor& ActorRef, FName Role)
 	return true;
 }
 
+bool FContextualAnimSceneBindings::SetRoleWarpTarget(const FName Role, const FName WarpTargetName, const FTransform& Transform)
+{
+	if (const FContextualAnimSceneBinding* Binding = FindBindingByRole(Role))
+	{
+		const_cast<FContextualAnimSceneBinding*>(Binding)->SetExternalWarpTargetTransform(WarpTargetName, Transform);
+		return true;
+	}
+	return false;
+}
+
 const FContextualAnimTrack& FContextualAnimSceneBindings::GetAnimTrackFromBinding(const FContextualAnimSceneBinding& Binding) const
 {
 	if (const FContextualAnimTrack* AnimTrack = GetSceneAsset()->GetAnimTrack(SectionIdx, AnimSetIdx, Binding.AnimTrackIdx))
diff --git a/Engine/Plugins/Experimental/Animation/ContextualAnimation/Source/ContextualAnimation/Public/ContextualAnimSceneActorComponent.h b/Engine/Plugins/Experimental/Animation/ContextualAnimation/Source/ContextualAnimation/Public/ContextualAnimSceneActorComponent.h
--- a/Engine/Plugins/Experimental/Animation/ContextualAnimation/Source/ContextualAnimation/Public/ContextualAnimSceneActorComponent.h
+++ b/Engine/Plugins/Experimental/Animation/ContextualAnimation/Source/ContextualAnimation/Public/ContextualAnimSceneActorComponent.h
@@ -56,6 +56,30 @@ struct FContextualAnimRepTransitionData : public FContextualAnimRepData
 	uint8 AnimSetIdx = 0;
 };
 
+USTRUCT(BlueprintType)
+struct FContextualAnimWarpTarget
+{
+	GENERATED_BODY()
+
+	FContextualAnimWarpTarget() = default;
+	FContextualAnimWarpTarget(const FName InRole, const FName InWarpTargetName, const FTransform& InTargetTransform)
+		: Role(InRole)
+		, TargetName(InWarpTargetName)
+		, TargetTransform(InTargetTransform)
+	{
+	}
+	
+	UPROPERTY(EditAnywhere, Category = "Default")
+	FName Role;
+
+	UPROPERTY(EditAnywhere, Category = "Default")
+	FName TargetName;
+
+	UPROPERTY(EditAnywhere, Category = "Default")
+	FTransform TargetTransform;
+};
+
+
 UCLASS(meta = (BlueprintSpawnableComponent))
 class CONTEXTUALANIMATION_API UContextualAnimSceneActorComponent : public UPrimitiveComponent, public IIKGoalCreatorInterface
 {
@@ -105,16 +129,21 @@ public:
 
 	UFUNCTION(BlueprintCallable, Category = "Contextual Anim|Scene Actor Component")
 	bool StartContextualAnimScene(const FContextualAnimSceneBindings& InBindings);
-
+	
 	UFUNCTION(BlueprintCallable, Category = "Contextual Anim|Scene Actor Component")
 	bool LateJoinContextualAnimScene(AActor* Actor, FName Role);
-
+	
 	UFUNCTION(BlueprintCallable, Category = "Contextual Anim|Scene Actor Component")
 	bool TransitionContextualAnimScene(FName SectionName);
-
+	
 	UFUNCTION(BlueprintCallable, Category = "Contextual Anim|Scene Actor Component")
 	bool TransitionSingleActor(int32 SectionIdx, int32 AnimSetIdx);
 
+	bool StartContextualAnimScene(const FContextualAnimSceneBindings& InBindings, const TArray<FContextualAnimWarpTarget>& WarpTargets);
+	bool LateJoinContextualAnimScene(AActor* Actor, FName Role, const TArray<FContextualAnimWarpTarget>& WarpTargets);
+	bool TransitionContextualAnimScene(FName SectionName, const TArray<FContextualAnimWarpTarget>& WarpTargets);
+	bool TransitionSingleActor(int32 SectionIdx, int32 AnimSetIdx, const TArray<FContextualAnimWarpTarget>& WarpTargets);
+
 	void EarlyOutContextualAnimScene();
 
 	bool IsOwnerLocallyControlled() const;
@@ -189,13 +218,13 @@ protected:
 	// @TODO: These two functions are going to replace OnJoinedScene and OnLeftScene
 	// main different is that these new functions are taking care of animation playback too
 
-	void JoinScene(const FContextualAnimSceneBindings& InBindings);
+	void JoinScene(const FContextualAnimSceneBindings& InBindings, const TArray<FContextualAnimWarpTarget>& WarpTargets);
 
 	void LeaveScene();
 
-	void LateJoinScene(const FContextualAnimSceneBindings& InBindings);
+	void LateJoinScene(const FContextualAnimSceneBindings& InBindings, const TArray<FContextualAnimWarpTarget>& WarpTargets);
 
-	bool HandleLateJoin(AActor* Actor, FName Role);
+	bool HandleLateJoin(AActor* Actor, FName Role, const TArray<FContextualAnimWarpTarget>& WarpTargets);
 
 	UFUNCTION(Server, Reliable, WithValidation)
 	void ServerStartContextualAnimScene(const FContextualAnimSceneBindings& InBindings);
diff --git a/Engine/Plugins/Experimental/Animation/ContextualAnimation/Source/ContextualAnimation/Public/ContextualAnimTypes.h b/Engine/Plugins/Experimental/Animation/ContextualAnimation/Source/ContextualAnimation/Public/ContextualAnimTypes.h
--- a/Engine/Plugins/Experimental/Animation/ContextualAnimation/Source/ContextualAnimation/Public/ContextualAnimTypes.h
+++ b/Engine/Plugins/Experimental/Animation/ContextualAnimation/Source/ContextualAnimation/Public/ContextualAnimTypes.h
@@ -343,6 +343,16 @@ struct CONTEXTUALANIMATION_API FContextualAnimSceneBindingContext
 
 	FVector GetVelocity() const;
 
+	void SetExternalWarpTargetTransform(const FName TargetName, const FTransform& Transform)
+	{
+		ExternalWarpTargetName = TargetName;
+		ExternalWarpTargetTransform = Transform;
+	}
+
+	bool HasExternalWarpTarget() const { return !ExternalWarpTargetName.IsNone(); }
+	FName GetExternalWarpTargetName() const { return ExternalWarpTargetName; }
+	const FTransform& GetExternalWarpTargetTransform() const { return ExternalWarpTargetTransform; }
+
 	//@TODO: Add accessors for GameplayTags
 
 private:
@@ -364,6 +374,12 @@ private:
 	TOptional<FVector> ExternalVelocity;
 
 	FGameplayTagContainer ExternalGameplayTags;
+
+	UPROPERTY()
+	FName ExternalWarpTargetName;
+	
+	UPROPERTY()
+	FTransform ExternalWarpTargetTransform;
 };
 
 /** Represent an actor bound to a role in the scene */
@@ -387,6 +403,12 @@ struct CONTEXTUALANIMATION_API FContextualAnimSceneBinding
 	
 	void SetAnimTrack(const FContextualAnimTrack& InAnimTrack);
 
+	void SetExternalWarpTargetTransform(const FName TargetName, const FTransform& Transform);
+
+	bool HasExternalWarpTarget() const { return Context.HasExternalWarpTarget(); }
+	FName GetExternalWarpTargetName() const { return Context.GetExternalWarpTargetName(); }
+	const FTransform& GetExternalWarpTargetTransform() const { return Context.GetExternalWarpTargetTransform(); }
+
 	/** Return the current playback time of the animation this actor is playing */
 	float GetAnimMontageTime() const;
 
@@ -450,6 +472,8 @@ struct CONTEXTUALANIMATION_API FContextualAnimSceneBindings
 
 	bool BindActorToRole(AActor& ActorRef, FName Role);
 
+	bool SetRoleWarpTarget(const FName Role, const FName WarpTargetName, const FTransform& Transform);
+
 	void CalculateWarpPoints(TArray<FContextualAnimWarpPoint>& OutWarpPoints) const;
 	bool CalculateWarpPoint(const FContextualAnimWarpPointDefinition& WarpPointDef, FContextualAnimWarpPoint& OutWarpPoint) const;
 
diff --git a/Engine/Plugins/Experimental/Animation/ContextualAnimation/Source/ContextualAnimationEditor/Private/ContextualAnimViewModel.cpp b/Engine/Plugins/Experimental/Animation/ContextualAnimation/Source/ContextualAnimationEditor/Private/ContextualAnimViewModel.cpp
--- a/Engine/Plugins/Experimental/Animation/ContextualAnimation/Source/ContextualAnimationEditor/Private/ContextualAnimViewModel.cpp
+++ b/Engine/Plugins/Experimental/Animation/ContextualAnimation/Source/ContextualAnimationEditor/Private/ContextualAnimViewModel.cpp
@@ -525,7 +525,7 @@ void FContextualAnimViewModel::RefreshPreviewScene()
 
 	// Pause all anims and disable auto blend out
 	const FContextualAnimSceneBinding* Leader = SceneBindings.GetSyncLeader();
-	if (Leader && Leader->GetSceneActorComponent()->StartContextualAnimScene(SceneBindings))
+	if (Leader && Leader->GetSceneActorComponent()->StartContextualAnimScene(SceneBindings, {}))
 	{
 		for (const FContextualAnimSceneBinding& Binding : SceneBindings)
 		{
@@ -744,7 +744,7 @@ void FContextualAnimViewModel::StartSimulation()
 	if (FContextualAnimSceneBindings::TryCreateBindings(*GetSceneAsset(), ActiveSectionIdx, RoleToActorMap, SceneBindings))
 	{
 		const FContextualAnimSceneBinding* Leader = SceneBindings.GetSyncLeader();
-		if (Leader && Leader->GetSceneActorComponent()->StartContextualAnimScene(SceneBindings))
+		if (Leader && Leader->GetSceneActorComponent()->StartContextualAnimScene(SceneBindings, {}))
 		{
 			SimulateModeState = ESimulateModeState::Playing;
 		}
