Contextual Anim: Added support for external warp targets

diff --git a/Engine/Plugins/Experimental/Animation/ContextualAnimation/Source/ContextualAnimation/Private/ContextualAnimSceneActorComponent.cpp b/Engine/Plugins/Experimental/Animation/ContextualAnimation/Source/ContextualAnimation/Private/ContextualAnimSceneActorComponent.cpp
--- a/Engine/Plugins/Experimental/Animation/ContextualAnimation/Source/ContextualAnimation/Private/ContextualAnimSceneActorComponent.cpp
+++ b/Engine/Plugins/Experimental/Animation/ContextualAnimation/Source/ContextualAnimation/Private/ContextualAnimSceneActorComponent.cpp
@@ -111,7 +111,7 @@ void UContextualAnimSceneActorComponent::PlayAnimation_Internal(UAnimSequenceBas
 	}
 }
 
-void UContextualAnimSceneActorComponent::AddOrUpdateWarpTargets(int32 SectionIdx, int32 AnimSetIdx)
+void UContextualAnimSceneActorComponent::AddOrUpdateWarpTargets(int32 SectionIdx, int32 AnimSetIdx, const TArray<FContextualAnimWarpTarget>& ExternalWarpTargets)
 {
 	// This is relevant only for character with motion warping comp
 	ACharacter* CharacterOwner = Cast<ACharacter>(GetOwner());
@@ -150,14 +150,18 @@ void UContextualAnimSceneActorComponent::AddOrUpdateWarpTargets(int32 SectionIdx
 			}
 		}
 
-		if (Binding->HasExternalWarpTarget())
+		const FName Role = Bindings.GetRoleFromBinding(*Binding);
+		for (const FContextualAnimWarpTarget& WarpTarget : ExternalWarpTargets)
 		{
-			MotionWarpComp->AddOrUpdateWarpTargetFromTransform(Binding->GetExternalWarpTargetName(), Binding->GetExternalWarpTargetTransform());
+			if (WarpTarget.Role == Role)
+			{
+				MotionWarpComp->AddOrUpdateWarpTargetFromTransform(WarpTarget.TargetName, FTransform(WarpTarget.TargetRotation, WarpTarget.TargetLocation));
+			}
 		}
 	}
 }
 
-bool UContextualAnimSceneActorComponent::LateJoinContextualAnimScene(AActor* Actor, FName Role, const TArray<FContextualAnimWarpTarget>& WarpTargets)
+bool UContextualAnimSceneActorComponent::LateJoinContextualAnimScene(AActor* Actor, FName Role, const TArray<FContextualAnimWarpTarget>& ExternalWarpTargets)
 {
 	if (!Bindings.IsValid())
 	{
@@ -173,7 +177,7 @@ bool UContextualAnimSceneActorComponent::LateJoinContextualAnimScene(AActor* Act
 		{
 			if (UContextualAnimSceneActorComponent* Comp = Leader->GetSceneActorComponent())
 			{
-				return Comp->LateJoinContextualAnimScene(Actor, Role, WarpTargets);
+				return Comp->LateJoinContextualAnimScene(Actor, Role, ExternalWarpTargets);
 			}
 		}
 	}
@@ -182,13 +186,14 @@ bool UContextualAnimSceneActorComponent::LateJoinContextualAnimScene(AActor* Act
 		*UEnum::GetValueAsString(TEXT("Engine.ENetRole"), GetOwner()->GetLocalRole()), *GetNameSafe(GetOwner()), Bindings.GetID(), Bindings.GetSectionIdx(), *GetNameSafe(Bindings.GetSceneAsset()), *GetNameSafe(Actor), *Role.ToString());
 
 	// Play animation and set state on this new actor that is joining us and update bindings for everyone else
-	if (HandleLateJoin(Actor, Role, WarpTargets))
+	if (HandleLateJoin(Actor, Role, ExternalWarpTargets))
 	{
 		// Replicate late join event. See OnRep_LateJoinData
 		if (GetOwner()->HasAuthority())
 		{
 			RepLateJoinData.Actor = Actor;
 			RepLateJoinData.Role = Role;
+			RepLateJoinData.ExternalWarpTargets = ExternalWarpTargets;
 			RepLateJoinData.IncrementRepCounter();
 			MARK_PROPERTY_DIRTY_FROM_NAME(UContextualAnimSceneActorComponent, RepLateJoinData, this);
 			GetOwner()->ForceNetUpdate();
@@ -200,7 +205,7 @@ bool UContextualAnimSceneActorComponent::LateJoinContextualAnimScene(AActor* Act
 	return false;
 }
 
-bool UContextualAnimSceneActorComponent::HandleLateJoin(AActor* Actor, FName Role, const TArray<FContextualAnimWarpTarget>& WarpTargets)
+bool UContextualAnimSceneActorComponent::HandleLateJoin(AActor* Actor, FName Role, const TArray<FContextualAnimWarpTarget>& ExternalWarpTargets)
 {
 	if (!Bindings.BindActorToRole(*Actor, Role))
 	{
@@ -227,14 +232,14 @@ bool UContextualAnimSceneActorComponent::HandleLateJoin(AActor* Actor, FName Rol
 	{
 		if (UContextualAnimSceneActorComponent* Comp = Binding->GetSceneActorComponent())
 		{
-			Comp->LateJoinScene(Bindings, WarpTargets);
+			Comp->LateJoinScene(Bindings, ExternalWarpTargets);
 		}
 	}
 
 	return true;
 }
 
-void UContextualAnimSceneActorComponent::LateJoinScene(const FContextualAnimSceneBindings& InBindings, const TArray<FContextualAnimWarpTarget>& WarpTargets)
+void UContextualAnimSceneActorComponent::LateJoinScene(const FContextualAnimSceneBindings& InBindings, const TArray<FContextualAnimWarpTarget>& ExternalWarpTargets)
 {
 	if (Bindings.IsValid())
 	{
@@ -251,11 +256,6 @@ void UContextualAnimSceneActorComponent::LateJoinScene(const FContextualAnimScen
 
 		Bindings = InBindings;
 
-		for (const FContextualAnimWarpTarget& WarpTarget : WarpTargets)
-		{
-			Bindings.SetRoleWarpTarget(WarpTarget.Role, WarpTarget.TargetName, WarpTarget.TargetTransform);
-		}
-
 		// For now when late joining an scene always play animation from first section
 		const int32 SectionIdx = 0;
 		const int32 AnimSetIdx = 0;
@@ -264,7 +264,7 @@ void UContextualAnimSceneActorComponent::LateJoinScene(const FContextualAnimScen
 
 		PlayAnimation_Internal(AnimTrack->Animation, 0.f, false);
 
-		AddOrUpdateWarpTargets(SectionIdx, AnimSetIdx);
+		AddOrUpdateWarpTargets(SectionIdx, AnimSetIdx, ExternalWarpTargets);
 
 		SetIgnoreCollisionWithOtherActors(true);
 
@@ -286,10 +286,10 @@ void UContextualAnimSceneActorComponent::OnRep_LateJoinData()
 		*UEnum::GetValueAsString(TEXT("Engine.ENetRole"), GetOwner()->GetLocalRole()), *GetNameSafe(GetOwner()), Bindings.GetID(), Bindings.GetSectionIdx(), *GetNameSafe(Bindings.GetSceneAsset()), *GetNameSafe(RepLateJoinData.Actor), *RepLateJoinData.Role.ToString());
 
 	// Play animation and set state on this new actor that is joining us and update bindings for everyone else
-	HandleLateJoin(RepLateJoinData.Actor, RepLateJoinData.Role, {});
+	HandleLateJoin(RepLateJoinData.Actor, RepLateJoinData.Role, RepLateJoinData.ExternalWarpTargets);
 }
 
-bool UContextualAnimSceneActorComponent::TransitionContextualAnimScene(FName SectionName, const TArray<FContextualAnimWarpTarget>& WarpTargets)
+bool UContextualAnimSceneActorComponent::TransitionContextualAnimScene(FName SectionName, const TArray<FContextualAnimWarpTarget>& ExternalWarpTargets)
 {
 	if (!GetOwner()->HasAuthority())
 	{
@@ -304,28 +304,24 @@ bool UContextualAnimSceneActorComponent::TransitionContextualAnimScene(FName Sec
 		{
 			if (UContextualAnimSceneActorComponent* Comp = Leader->GetSceneActorComponent())
 			{
-				return Comp->TransitionContextualAnimScene(SectionName, WarpTargets);
+				return Comp->TransitionContextualAnimScene(SectionName, ExternalWarpTargets);
 			}
 		}
 	}
 
 	if (const FContextualAnimSceneBinding* OwnerBinding = Bindings.FindBindingByActor(GetOwner()))
 	{
-		for (const FContextualAnimWarpTarget& WarpTarget : WarpTargets)
-		{
-			Bindings.SetRoleWarpTarget(WarpTarget.Role, WarpTarget.TargetName, WarpTarget.TargetTransform);
-		}
-
 		const int32 SectionIdx = Bindings.GetSceneAsset()->GetSectionIndex(SectionName);
 		if (SectionIdx != INDEX_NONE)
 		{
 			UE_LOG(LogContextualAnim, Log, TEXT("%-21s UContextualAnimSceneActorComponent::TransitionTo Actor: %s SectionName: %s"),
 				*UEnum::GetValueAsString(TEXT("Engine.ENetRole"), GetOwner()->GetLocalRole()), *GetNameSafe(GetOwner()), *SectionName.ToString());
 
-			HandleTransitionEveryone(SectionIdx, 0);
+			HandleTransitionEveryone(SectionIdx, 0, ExternalWarpTargets);
 
 			RepTransitionData.SectionIdx = SectionIdx;
 			RepTransitionData.AnimSetIdx = 0;
+			RepTransitionData.ExternalWarpTargets = ExternalWarpTargets;
 			RepTransitionData.IncrementRepCounter();
 			MARK_PROPERTY_DIRTY_FROM_NAME(UContextualAnimSceneActorComponent, RepTransitionData, this);
 			GetOwner()->ForceNetUpdate();
@@ -337,11 +333,11 @@ bool UContextualAnimSceneActorComponent::TransitionContextualAnimScene(FName Sec
 	return false;
 }
 
-void UContextualAnimSceneActorComponent::HandleTransitionEveryone(int32 NewSectionIdx, int32 NewAnimSetIdx)
+void UContextualAnimSceneActorComponent::HandleTransitionEveryone(int32 NewSectionIdx, int32 NewAnimSetIdx, const TArray<FContextualAnimWarpTarget>& ExternalWarpTargets)
 {
 	// Update Bindings internal data and play new animation for the leader first
 	// Note that for now we always transition to the first set in the section. We could run selection criteria here too but keeping it simple for now
-	HandleTransitionSelf(NewSectionIdx, NewAnimSetIdx);
+	HandleTransitionSelf(NewSectionIdx, NewAnimSetIdx, ExternalWarpTargets);
 
 	// And now the same for everyone else
 	for (const FContextualAnimSceneBinding& Binding : Bindings)
@@ -350,13 +346,13 @@ void UContextualAnimSceneActorComponent::HandleTransitionEveryone(int32 NewSecti
 		{
 			if(UContextualAnimSceneActorComponent* Comp = Binding.GetSceneActorComponent())
 			{
-				Comp->HandleTransitionSelf(NewSectionIdx, NewAnimSetIdx);
+				Comp->HandleTransitionSelf(NewSectionIdx, NewAnimSetIdx, ExternalWarpTargets);
 			}
 		}
 	}
 }
 
-void UContextualAnimSceneActorComponent::HandleTransitionSelf(int32 NewSectionIdx, int32 NewAnimSetIdx)
+void UContextualAnimSceneActorComponent::HandleTransitionSelf(int32 NewSectionIdx, int32 NewAnimSetIdx, const TArray<FContextualAnimWarpTarget>& ExternalWarpTargets)
 {
 	// Update bindings internal data so it points to the new section and new anim set
 	Bindings.TransitionTo(NewSectionIdx, NewAnimSetIdx);
@@ -366,10 +362,10 @@ void UContextualAnimSceneActorComponent::HandleTransitionSelf(int32 NewSectionId
 	const FContextualAnimTrack& AnimTrack = Bindings.GetAnimTrackFromBinding(*Bindings.FindBindingByActor(GetOwner()));
 	PlayAnimation_Internal(AnimTrack.Animation, 0.f, true);
 
-	AddOrUpdateWarpTargets(NewSectionIdx, NewAnimSetIdx);
+	AddOrUpdateWarpTargets(NewSectionIdx, NewAnimSetIdx, ExternalWarpTargets);
 }
 
-bool UContextualAnimSceneActorComponent::TransitionSingleActor(int32 SectionIdx, int32 AnimSetIdx, const TArray<FContextualAnimWarpTarget>& WarpTargets)
+bool UContextualAnimSceneActorComponent::TransitionSingleActor(int32 SectionIdx, int32 AnimSetIdx, const TArray<FContextualAnimWarpTarget>& ExternalWarpTargets)
 {
 	if (!GetOwner()->HasAuthority())
 	{
@@ -378,11 +374,6 @@ bool UContextualAnimSceneActorComponent::TransitionSingleActor(int32 SectionIdx,
 
 	if (const FContextualAnimSceneBinding* OwnerBinding = Bindings.FindBindingByActor(GetOwner()))
 	{
-		for (const FContextualAnimWarpTarget& WarpTarget : WarpTargets)
-		{
-			Bindings.SetRoleWarpTarget(WarpTarget.Role, WarpTarget.TargetName, WarpTarget.TargetTransform);
-		}
-
 		if (const UContextualAnimSceneAsset* Asset = Bindings.GetSceneAsset())
 		{
 			const FContextualAnimTrack* AnimTrack = Asset->GetAnimTrack(SectionIdx, AnimSetIdx, Bindings.GetRoleFromBinding(*OwnerBinding));
@@ -393,10 +384,11 @@ bool UContextualAnimSceneActorComponent::TransitionSingleActor(int32 SectionIdx,
 
 				PlayAnimation_Internal(AnimTrack->Animation, 0.f, false);
 
-				AddOrUpdateWarpTargets(SectionIdx, AnimSetIdx);
+				AddOrUpdateWarpTargets(SectionIdx, AnimSetIdx, ExternalWarpTargets);
 
 				RepTransitionSingleActorData.SectionIdx = SectionIdx;
 				RepTransitionSingleActorData.AnimSetIdx = AnimSetIdx;
+				RepTransitionSingleActorData.ExternalWarpTargets = ExternalWarpTargets;
 				RepTransitionSingleActorData.IncrementRepCounter();
 				MARK_PROPERTY_DIRTY_FROM_NAME(UContextualAnimSceneActorComponent, RepTransitionSingleActorData, this);
 				GetOwner()->ForceNetUpdate();
@@ -422,7 +414,8 @@ void UContextualAnimSceneActorComponent::OnRep_RepTransitionSingleActor()
 			if (AnimTrack && AnimTrack->Animation)
 			{
 				PlayAnimation_Internal(AnimTrack->Animation, 0.f, false);
-				AddOrUpdateWarpTargets(RepTransitionSingleActorData.SectionIdx, RepTransitionSingleActorData.AnimSetIdx);
+
+				AddOrUpdateWarpTargets(RepTransitionSingleActorData.SectionIdx, RepTransitionSingleActorData.AnimSetIdx, RepTransitionSingleActorData.ExternalWarpTargets);
 			}
 		}
 	}
@@ -448,7 +441,7 @@ bool UContextualAnimSceneActorComponent::TransitionSingleActor(int32 SectionIdx,
 	return TransitionSingleActor(SectionIdx, AnimSetIdx, {});
 }
 
-bool UContextualAnimSceneActorComponent::StartContextualAnimScene(const FContextualAnimSceneBindings& InBindings, const TArray<FContextualAnimWarpTarget>& WarpTargets)
+bool UContextualAnimSceneActorComponent::StartContextualAnimScene(const FContextualAnimSceneBindings& InBindings, const TArray<FContextualAnimWarpTarget>& ExternalWarpTargets)
 {
 	UE_LOG(LogContextualAnim, Log, TEXT("%-21s UContextualAnimSceneActorComponent::StartContextualAnim Actor: %s"),
 		*UEnum::GetValueAsString(TEXT("Engine.ENetRole"), GetOwner()->GetLocalRole()), *GetNameSafe(GetOwner()));
@@ -458,7 +451,7 @@ bool UContextualAnimSceneActorComponent::StartContextualAnimScene(const FContext
 	{
 		if (GetOwner()->HasAuthority())
 		{
-			JoinScene(InBindings, WarpTargets);
+			JoinScene(InBindings, ExternalWarpTargets);
 
 			for (const FContextualAnimSceneBinding& Binding : InBindings)
 			{
@@ -466,12 +459,13 @@ bool UContextualAnimSceneActorComponent::StartContextualAnimScene(const FContext
 				{
 					if (UContextualAnimSceneActorComponent* Comp = Binding.GetSceneActorComponent())
 					{
-						Comp->JoinScene(InBindings, WarpTargets);
+						Comp->JoinScene(InBindings, ExternalWarpTargets);
 					}
 				}
 			}
 
-			RepBindings = InBindings;
+			RepBindings.Bindings = InBindings;
+			RepBindings.ExternalWarpTargets = ExternalWarpTargets;
 			MARK_PROPERTY_DIRTY_FROM_NAME(UContextualAnimSceneActorComponent, RepBindings, this);
 			GetOwner()->ForceNetUpdate();
 
@@ -479,7 +473,7 @@ bool UContextualAnimSceneActorComponent::StartContextualAnimScene(const FContext
 		}
 		else if (GetOwner()->GetLocalRole() == ROLE_AutonomousProxy)
 		{
-			JoinScene(InBindings, WarpTargets);
+			JoinScene(InBindings, ExternalWarpTargets);
 
 			ServerStartContextualAnimScene(InBindings);
 
@@ -519,9 +513,10 @@ void UContextualAnimSceneActorComponent::EarlyOutContextualAnimScene()
 				// If we are on the server, rep bindings to stop animation on simulated proxies
 				if (GetOwner()->HasAuthority())
 				{
-					if (RepBindings.IsValid())
+					if (RepBindings.Bindings.IsValid())
 					{
-						RepBindings.Clear();
+						RepBindings.Bindings.Clear();
+						RepBindings.ExternalWarpTargets.Empty();
 						MARK_PROPERTY_DIRTY_FROM_NAME(UContextualAnimSceneActorComponent, RepBindings, this);
 
 						GetOwner()->ForceNetUpdate();
@@ -552,37 +547,37 @@ void UContextualAnimSceneActorComponent::OnRep_TransitionData()
 	UE_LOG(LogContextualAnim, Verbose, TEXT("%-21s UContextualAnimSceneActorComponent::OnRep_TransitionData Actor: %s"),
 		*UEnum::GetValueAsString(TEXT("Engine.ENetRole"), GetOwner()->GetLocalRole()), *GetNameSafe(GetOwner()));
 
-	HandleTransitionEveryone(RepTransitionData.SectionIdx, RepTransitionData.AnimSetIdx);
+	HandleTransitionEveryone(RepTransitionData.SectionIdx, RepTransitionData.AnimSetIdx, RepTransitionData.ExternalWarpTargets);
 }
 
 void UContextualAnimSceneActorComponent::OnRep_Bindings()
 {
-	UE_LOG(LogContextualAnim, Verbose, TEXT("%-21s UContextualAnimSceneActorComponent::OnRep_Bindings Actor: %s RepBindings Id: %d Num: %d Bindings Id: %d Num: %d"),
-		*UEnum::GetValueAsString(TEXT("Engine.ENetRole"), GetOwner()->GetLocalRole()), *GetNameSafe(GetOwner()), RepBindings.GetID(), RepBindings.Num(), Bindings.GetID(), Bindings.Num());
+	UE_LOG(LogContextualAnim, Verbose, TEXT("%-21s UContextualAnimSceneActorComponent::OnRep_Bindings Actor: %s Rep Bindings Id: %d Num: %d Current Bindings Id: %d Num: %d"),
+		*UEnum::GetValueAsString(TEXT("Engine.ENetRole"), GetOwner()->GetLocalRole()), *GetNameSafe(GetOwner()), RepBindings.Bindings.GetID(), RepBindings.Bindings.Num(), Bindings.GetID(), Bindings.Num());
 
 	// The owner of this component started an interaction on the server
-	if (RepBindings.IsValid())
+	if (RepBindings.Bindings.IsValid())
 	{
-		const FContextualAnimSceneBinding* OwnerBinding = RepBindings.FindBindingByActor(GetOwner());
+		const FContextualAnimSceneBinding* OwnerBinding = RepBindings.Bindings.FindBindingByActor(GetOwner());
 		if (ensureAlways(OwnerBinding))
 		{
 			// Join the scene (start playing animation, etc.)
 			if (GetOwner()->GetLocalRole() != ROLE_AutonomousProxy)
 			{
-				JoinScene(RepBindings, {});
+				JoinScene(RepBindings.Bindings, RepBindings.ExternalWarpTargets);
 			}
 
 			// RepBindings is only replicated from the initiator of the action.
 			// So now we have to tell everyone else involved in the interaction to join us
 			// @TODO: For now this assumes that all the actors will start playing the animation at the same time. 
 			// We will expand this in the future to allow 'late' join
-			for (const FContextualAnimSceneBinding& Binding : RepBindings)
+			for (const FContextualAnimSceneBinding& Binding : RepBindings.Bindings)
 			{
 				if (Binding.GetActor() != GetOwner())
 				{
 					if (UContextualAnimSceneActorComponent* Comp = Binding.GetSceneActorComponent())
 					{
-						Comp->JoinScene(RepBindings, {});
+						Comp->JoinScene(RepBindings.Bindings, RepBindings.ExternalWarpTargets);
 					}
 				}
 			}
@@ -594,7 +589,7 @@ void UContextualAnimSceneActorComponent::OnRep_Bindings()
 		// In this case we don't want to tell everyone else to also leave the scene since there is very common for the initiator, 
 		// specially if is player character, to end the animation earlier for responsiveness
 		// It is more likely this will do nothing since we listen to montage end also on Simulated Proxies to 'predict' the end of the interaction.
-		if (RepBindings.GetID() == Bindings.GetID() && GetOwner()->GetLocalRole() != ROLE_AutonomousProxy)
+		if (RepBindings.Bindings.GetID() == Bindings.GetID() && GetOwner()->GetLocalRole() != ROLE_AutonomousProxy)
 		{
 			LeaveScene();
 		}
@@ -720,7 +715,7 @@ void UContextualAnimSceneActorComponent::OnLeftScene()
 	}
 }
 
-void UContextualAnimSceneActorComponent::JoinScene(const FContextualAnimSceneBindings& InBindings, const TArray<FContextualAnimWarpTarget>& WarpTargets)
+void UContextualAnimSceneActorComponent::JoinScene(const FContextualAnimSceneBindings& InBindings, const TArray<FContextualAnimWarpTarget>& ExternalWarpTargets)
 {
 	if (Bindings.IsValid())
 	{
@@ -734,15 +729,10 @@ void UContextualAnimSceneActorComponent::JoinScene(const FContextualAnimSceneBin
 
 		Bindings = InBindings;
 
-		for (const FContextualAnimWarpTarget& WarpTarget : WarpTargets)
-		{
-			Bindings.SetRoleWarpTarget(WarpTarget.Role, WarpTarget.TargetName, WarpTarget.TargetTransform);
-		}
-
 		const FContextualAnimTrack& AnimTrack = Bindings.GetAnimTrackFromBinding(*Binding);
 		PlayAnimation_Internal(AnimTrack.Animation, 0.f, true);
 
-		AddOrUpdateWarpTargets(AnimTrack.SectionIdx, AnimTrack.AnimSetIdx);
+		AddOrUpdateWarpTargets(AnimTrack.SectionIdx, AnimTrack.AnimSetIdx, ExternalWarpTargets);
 
 		// Disable collision between actors so they can align perfectly
 		SetIgnoreCollisionWithOtherActors(true);
@@ -857,9 +847,10 @@ void UContextualAnimSceneActorComponent::OnMontageBlendingOut(UAnimMontage* Mont
 		if (GetOwner()->HasAuthority())
 		{
 			// Rep empty bindings if we were the initiator of this interaction.
-			if (RepBindings.IsValid())
+			if (RepBindings.Bindings.IsValid())
 			{
-				RepBindings.Clear();
+				RepBindings.Bindings.Clear();
+				RepBindings.ExternalWarpTargets.Empty();
 				MARK_PROPERTY_DIRTY_FROM_NAME(UContextualAnimSceneActorComponent, RepBindings, this);
 
 				GetOwner()->ForceNetUpdate();
diff --git a/Engine/Plugins/Experimental/Animation/ContextualAnimation/Source/ContextualAnimation/Private/ContextualAnimTypes.cpp b/Engine/Plugins/Experimental/Animation/ContextualAnimation/Source/ContextualAnimation/Private/ContextualAnimTypes.cpp
--- a/Engine/Plugins/Experimental/Animation/ContextualAnimation/Source/ContextualAnimation/Private/ContextualAnimTypes.cpp
+++ b/Engine/Plugins/Experimental/Animation/ContextualAnimation/Source/ContextualAnimation/Private/ContextualAnimTypes.cpp
@@ -274,11 +274,6 @@ void FContextualAnimSceneBinding::SetAnimTrack(const FContextualAnimTrack& InAni
 	AnimTrackIdx = InAnimTrack.AnimTrackIdx;
 }
 
-void FContextualAnimSceneBinding::SetExternalWarpTargetTransform(const FName TargetName, const FTransform& Transform)
-{
-	Context.SetExternalWarpTargetTransform(TargetName, Transform);
-}
-
 FAnimMontageInstance* FContextualAnimSceneBinding::GetAnimMontageInstance() const
 {
 	if (UAnimInstance* AnimInstance = GetAnimInstance())
@@ -391,16 +386,6 @@ bool FContextualAnimSceneBindings::BindActorToRole(AActor& ActorRef, FName Role)
 	return true;
 }
 
-bool FContextualAnimSceneBindings::SetRoleWarpTarget(const FName Role, const FName WarpTargetName, const FTransform& Transform)
-{
-	if (const FContextualAnimSceneBinding* Binding = FindBindingByRole(Role))
-	{
-		const_cast<FContextualAnimSceneBinding*>(Binding)->SetExternalWarpTargetTransform(WarpTargetName, Transform);
-		return true;
-	}
-	return false;
-}
-
 const FContextualAnimTrack& FContextualAnimSceneBindings::GetAnimTrackFromBinding(const FContextualAnimSceneBinding& Binding) const
 {
 	if (const FContextualAnimTrack* AnimTrack = GetSceneAsset()->GetAnimTrack(SectionIdx, AnimSetIdx, Binding.AnimTrackIdx))
diff --git a/Engine/Plugins/Experimental/Animation/ContextualAnimation/Source/ContextualAnimation/Public/ContextualAnimSceneActorComponent.h b/Engine/Plugins/Experimental/Animation/ContextualAnimation/Source/ContextualAnimation/Public/ContextualAnimSceneActorComponent.h
--- a/Engine/Plugins/Experimental/Animation/ContextualAnimation/Source/ContextualAnimation/Public/ContextualAnimSceneActorComponent.h
+++ b/Engine/Plugins/Experimental/Animation/ContextualAnimation/Source/ContextualAnimation/Public/ContextualAnimSceneActorComponent.h
@@ -15,6 +15,33 @@ class UAnimSequenceBase;
 
 DECLARE_DYNAMIC_MULTICAST_DELEGATE_OneParam(FContextualAnimSceneActorCompDelegate, class UContextualAnimSceneActorComponent*, SceneActorComponent);
 
+USTRUCT(BlueprintType)
+struct FContextualAnimWarpTarget
+{
+	GENERATED_BODY()
+
+	FContextualAnimWarpTarget() = default;
+	FContextualAnimWarpTarget(const FName InRole, const FName InWarpTargetName, const FTransform& InTargetTransform)
+		: Role(InRole)
+		, TargetName(InWarpTargetName)
+		, TargetLocation(InTargetTransform.GetLocation())
+		, TargetRotation(InTargetTransform.GetRotation())
+	{
+	}
+
+	UPROPERTY(EditAnywhere, Category = "Default")
+	FName Role;
+
+	UPROPERTY(EditAnywhere, Category = "Default")
+	FName TargetName;
+
+	UPROPERTY(EditAnywhere, Category = "Default")
+	FVector TargetLocation;
+
+	UPROPERTY(EditAnywhere, Category = "Default")
+	FQuat TargetRotation;
+};
+
 /** Base struct for replicated data with a rep counter */
 USTRUCT()
 struct FContextualAnimRepData
@@ -28,6 +55,19 @@ struct FContextualAnimRepData
 	void IncrementRepCounter();
 };
 
+/** Used to replicate start/stop contextual anim events */
+USTRUCT()
+struct FContextualAnimRepBindingsData : public FContextualAnimRepData
+{
+	GENERATED_BODY()
+
+	UPROPERTY()
+	FContextualAnimSceneBindings Bindings;
+
+	UPROPERTY()
+	TArray<FContextualAnimWarpTarget> ExternalWarpTargets;
+};
+
 /** Used to replicate a late join event */
 USTRUCT()
 struct FContextualAnimRepLateJoinData : public FContextualAnimRepData
@@ -41,9 +81,12 @@ struct FContextualAnimRepLateJoinData : public FContextualAnimRepData
 	/** Role in the interaction the actor is gonna play */
 	UPROPERTY()
 	FName Role = NAME_None;
+
+	UPROPERTY()
+	TArray<FContextualAnimWarpTarget> ExternalWarpTargets;
 };
 
-/** Used to replicate a transition to a new AnimSet in the interaction */
+/** Used to transition events */
 USTRUCT()
 struct FContextualAnimRepTransitionData : public FContextualAnimRepData
 {
@@ -54,32 +97,11 @@ struct FContextualAnimRepTransitionData : public FContextualAnimRepData
 
 	UPROPERTY()
 	uint8 AnimSetIdx = 0;
-};
-
-USTRUCT(BlueprintType)
-struct FContextualAnimWarpTarget
-{
-	GENERATED_BODY()
-
-	FContextualAnimWarpTarget() = default;
-	FContextualAnimWarpTarget(const FName InRole, const FName InWarpTargetName, const FTransform& InTargetTransform)
-		: Role(InRole)
-		, TargetName(InWarpTargetName)
-		, TargetTransform(InTargetTransform)
-	{
-	}
-	
-	UPROPERTY(EditAnywhere, Category = "Default")
-	FName Role;
-
-	UPROPERTY(EditAnywhere, Category = "Default")
-	FName TargetName;
 
-	UPROPERTY(EditAnywhere, Category = "Default")
-	FTransform TargetTransform;
+	UPROPERTY()
+	TArray<FContextualAnimWarpTarget> ExternalWarpTargets;
 };
 
-
 UCLASS(meta = (BlueprintSpawnableComponent))
 class CONTEXTUALANIMATION_API UContextualAnimSceneActorComponent : public UPrimitiveComponent, public IIKGoalCreatorInterface
 {
@@ -139,10 +161,10 @@ public:
 	UFUNCTION(BlueprintCallable, Category = "Contextual Anim|Scene Actor Component")
 	bool TransitionSingleActor(int32 SectionIdx, int32 AnimSetIdx);
 
-	bool StartContextualAnimScene(const FContextualAnimSceneBindings& InBindings, const TArray<FContextualAnimWarpTarget>& WarpTargets);
-	bool LateJoinContextualAnimScene(AActor* Actor, FName Role, const TArray<FContextualAnimWarpTarget>& WarpTargets);
-	bool TransitionContextualAnimScene(FName SectionName, const TArray<FContextualAnimWarpTarget>& WarpTargets);
-	bool TransitionSingleActor(int32 SectionIdx, int32 AnimSetIdx, const TArray<FContextualAnimWarpTarget>& WarpTargets);
+	bool StartContextualAnimScene(const FContextualAnimSceneBindings& InBindings, const TArray<FContextualAnimWarpTarget>& ExternalWarpTargets);
+	bool LateJoinContextualAnimScene(AActor* Actor, FName Role, const TArray<FContextualAnimWarpTarget>& ExternalWarpTargets);
+	bool TransitionContextualAnimScene(FName SectionName, const TArray<FContextualAnimWarpTarget>& ExternalWarpTargets);
+	bool TransitionSingleActor(int32 SectionIdx, int32 AnimSetIdx, const TArray<FContextualAnimWarpTarget>& ExternalWarpTargets);
 
 	void EarlyOutContextualAnimScene();
 
@@ -155,7 +177,7 @@ protected:
 	 * This gets replicated only from the initiator of the action and then set on all the other members of the interaction
 	 */
 	UPROPERTY(Transient, ReplicatedUsing = OnRep_Bindings)
-	FContextualAnimSceneBindings RepBindings;
+	FContextualAnimRepBindingsData RepBindings;
 
 	UPROPERTY(Transient, ReplicatedUsing = OnRep_LateJoinData)
 	FContextualAnimRepLateJoinData RepLateJoinData;
@@ -218,13 +240,13 @@ protected:
 	// @TODO: These two functions are going to replace OnJoinedScene and OnLeftScene
 	// main different is that these new functions are taking care of animation playback too
 
-	void JoinScene(const FContextualAnimSceneBindings& InBindings, const TArray<FContextualAnimWarpTarget>& WarpTargets);
+	void JoinScene(const FContextualAnimSceneBindings& InBindings, const TArray<FContextualAnimWarpTarget>& ExternalWarpTargets);
 
 	void LeaveScene();
 
-	void LateJoinScene(const FContextualAnimSceneBindings& InBindings, const TArray<FContextualAnimWarpTarget>& WarpTargets);
+	void LateJoinScene(const FContextualAnimSceneBindings& InBindings, const TArray<FContextualAnimWarpTarget>& ExternalWarpTargets);
 
-	bool HandleLateJoin(AActor* Actor, FName Role, const TArray<FContextualAnimWarpTarget>& WarpTargets);
+	bool HandleLateJoin(AActor* Actor, FName Role, const TArray<FContextualAnimWarpTarget>& ExternalWarpTargets);
 
 	UFUNCTION(Server, Reliable, WithValidation)
 	void ServerStartContextualAnimScene(const FContextualAnimSceneBindings& InBindings);
@@ -234,11 +256,11 @@ protected:
 
 	void PlayAnimation_Internal(UAnimSequenceBase* Animation, float StartTime, bool bSyncPlaybackTime);
 
-	void AddOrUpdateWarpTargets(int32 SectionIdx, int32 AnimSetIdx);
+	void AddOrUpdateWarpTargets(int32 SectionIdx, int32 AnimSetIdx, const TArray<FContextualAnimWarpTarget>& ExternalWarpTargets);
 
-	void HandleTransitionSelf(int32 NewSectionIdx, int32 NewAnimSetIdx);
+	void HandleTransitionSelf(int32 NewSectionIdx, int32 NewAnimSetIdx, const TArray<FContextualAnimWarpTarget>& ExternalWarpTargets);
 
-	void HandleTransitionEveryone(int32 NewSectionIdx, int32 NewAnimSetIdx);
+	void HandleTransitionEveryone(int32 NewSectionIdx, int32 NewAnimSetIdx, const TArray<FContextualAnimWarpTarget>& ExternalWarpTargets);
 
 private:
 
diff --git a/Engine/Plugins/Experimental/Animation/ContextualAnimation/Source/ContextualAnimation/Public/ContextualAnimTypes.h b/Engine/Plugins/Experimental/Animation/ContextualAnimation/Source/ContextualAnimation/Public/ContextualAnimTypes.h
--- a/Engine/Plugins/Experimental/Animation/ContextualAnimation/Source/ContextualAnimation/Public/ContextualAnimTypes.h
+++ b/Engine/Plugins/Experimental/Animation/ContextualAnimation/Source/ContextualAnimation/Public/ContextualAnimTypes.h
@@ -343,16 +343,6 @@ struct CONTEXTUALANIMATION_API FContextualAnimSceneBindingContext
 
 	FVector GetVelocity() const;
 
-	void SetExternalWarpTargetTransform(const FName TargetName, const FTransform& Transform)
-	{
-		ExternalWarpTargetName = TargetName;
-		ExternalWarpTargetTransform = Transform;
-	}
-
-	bool HasExternalWarpTarget() const { return !ExternalWarpTargetName.IsNone(); }
-	FName GetExternalWarpTargetName() const { return ExternalWarpTargetName; }
-	const FTransform& GetExternalWarpTargetTransform() const { return ExternalWarpTargetTransform; }
-
 	//@TODO: Add accessors for GameplayTags
 
 private:
@@ -374,12 +364,6 @@ private:
 	TOptional<FVector> ExternalVelocity;
 
 	FGameplayTagContainer ExternalGameplayTags;
-
-	UPROPERTY()
-	FName ExternalWarpTargetName;
-	
-	UPROPERTY()
-	FTransform ExternalWarpTargetTransform;
 };
 
 /** Represent an actor bound to a role in the scene */
@@ -403,12 +387,6 @@ struct CONTEXTUALANIMATION_API FContextualAnimSceneBinding
 	
 	void SetAnimTrack(const FContextualAnimTrack& InAnimTrack);
 
-	void SetExternalWarpTargetTransform(const FName TargetName, const FTransform& Transform);
-
-	bool HasExternalWarpTarget() const { return Context.HasExternalWarpTarget(); }
-	FName GetExternalWarpTargetName() const { return Context.GetExternalWarpTargetName(); }
-	const FTransform& GetExternalWarpTargetTransform() const { return Context.GetExternalWarpTargetTransform(); }
-
 	/** Return the current playback time of the animation this actor is playing */
 	float GetAnimMontageTime() const;
 
