Contextual Anim: Fixed bug where forcing to stop the animation will not replicate the event if the actors was not the initiator of the interaction, potentially causing simulated proxy to get stuck in a looping animations forever.

diff --git a/Engine/Plugins/Experimental/Animation/ContextualAnimation/Source/ContextualAnimation/Private/ContextualAnimSceneActorComponent.cpp b/Engine/Plugins/Experimental/Animation/ContextualAnimation/Source/ContextualAnimation/Private/ContextualAnimSceneActorComponent.cpp
--- a/Engine/Plugins/Experimental/Animation/ContextualAnimation/Source/ContextualAnimation/Private/ContextualAnimSceneActorComponent.cpp
+++ b/Engine/Plugins/Experimental/Animation/ContextualAnimation/Source/ContextualAnimation/Private/ContextualAnimSceneActorComponent.cpp
@@ -403,19 +403,33 @@ bool UContextualAnimSceneActorComponent::TransitionSingleActor(int32 SectionIdx,
 
 void UContextualAnimSceneActorComponent::OnRep_RepTransitionSingleActor()
 {
-	UE_LOG(LogContextualAnim, Verbose, TEXT("%-21s UContextualAnimSceneActorComponent::OnRep_RepTransitionSingleActor Owner: %s SectionIdx: %s AnimSetIdx: %d"),
+	UE_LOG(LogContextualAnim, Verbose, TEXT("%-21s UContextualAnimSceneActorComponent::OnRep_RepTransitionSingleActor Owner: %s SectionIdx: %d AnimSetIdx: %d"),
 		*UEnum::GetValueAsString(TEXT("Engine.ENetRole"), GetOwner()->GetLocalRole()), *GetNameSafe(GetOwner()), RepTransitionSingleActorData.SectionIdx, RepTransitionSingleActorData.AnimSetIdx);
 
 	if (const FContextualAnimSceneBinding* OwnerBinding = Bindings.FindBindingByActor(GetOwner()))
 	{
-		if (const UContextualAnimSceneAsset* Asset = Bindings.GetSceneAsset())
+		if (RepTransitionSingleActorData.SectionIdx != MAX_uint8 && RepTransitionSingleActorData.AnimSetIdx != MAX_uint8)
 		{
-			const FContextualAnimTrack* AnimTrack = Asset->GetAnimTrack(RepTransitionSingleActorData.SectionIdx, RepTransitionSingleActorData.AnimSetIdx, Bindings.GetRoleFromBinding(*OwnerBinding));
-			if (AnimTrack && AnimTrack->Animation)
+			if (const UContextualAnimSceneAsset* Asset = Bindings.GetSceneAsset())
 			{
-				PlayAnimation_Internal(AnimTrack->Animation, 0.f, false);
+				const FContextualAnimTrack* AnimTrack = Asset->GetAnimTrack(RepTransitionSingleActorData.SectionIdx, RepTransitionSingleActorData.AnimSetIdx, Bindings.GetRoleFromBinding(*OwnerBinding));
+				if (AnimTrack && AnimTrack->Animation)
+				{
+					PlayAnimation_Internal(AnimTrack->Animation, 0.f, false);
 
-				AddOrUpdateWarpTargets(RepTransitionSingleActorData.SectionIdx, RepTransitionSingleActorData.AnimSetIdx, RepTransitionSingleActorData.ExternalWarpTargets);
+					AddOrUpdateWarpTargets(RepTransitionSingleActorData.SectionIdx, RepTransitionSingleActorData.AnimSetIdx, RepTransitionSingleActorData.ExternalWarpTargets);
+				}
+			}
+		}
+		else
+		{
+			// RepTransitionSingleActorData with invalid indices is replicated when the animation ends
+			// In this case we don't want to tell everyone else to also leave the scene since there is very common for the initiator, 
+			// specially if is player character, to end the animation earlier for responsiveness
+			// It is more likely this will do nothing since we listen to montage end also on Simulated Proxies to 'predict' the end of the interaction.
+			if (GetOwner()->GetLocalRole() != ROLE_AutonomousProxy)
+			{
+				LeaveScene();
 			}
 		}
 	}
@@ -510,17 +524,13 @@ void UContextualAnimSceneActorComponent::EarlyOutContextualAnimScene()
 				// Stop animation.
 				LeaveScene();
 
-				// If we are on the server, rep bindings to stop animation on simulated proxies
+				// If we are on the server, rep the event to stop animation on simulated proxies
 				if (GetOwner()->HasAuthority())
 				{
-					if (RepBindings.Bindings.IsValid())
-					{
-						RepBindings.Bindings.Clear();
-						RepBindings.ExternalWarpTargets.Empty();
-						MARK_PROPERTY_DIRTY_FROM_NAME(UContextualAnimSceneActorComponent, RepBindings, this);
-
-						GetOwner()->ForceNetUpdate();
-					}
+					RepTransitionSingleActorData.Reset();
+					RepTransitionSingleActorData.IncrementRepCounter();
+					MARK_PROPERTY_DIRTY_FROM_NAME(UContextualAnimSceneActorComponent, RepTransitionSingleActorData, this);
+					GetOwner()->ForceNetUpdate();
 				}
 				// If local player, tell the server to stop the animation too
 				else if (GetOwner()->GetLocalRole() == ROLE_AutonomousProxy)
@@ -583,17 +593,6 @@ void UContextualAnimSceneActorComponent::OnRep_Bindings()
 			}
 		}
 	}
-	else
-	{	
-		// Empty bindings is replicated by the initiator of the interaction when the animation ends
-		// In this case we don't want to tell everyone else to also leave the scene since there is very common for the initiator, 
-		// specially if is player character, to end the animation earlier for responsiveness
-		// It is more likely this will do nothing since we listen to montage end also on Simulated Proxies to 'predict' the end of the interaction.
-		if (RepBindings.Bindings.GetID() == Bindings.GetID() && GetOwner()->GetLocalRole() != ROLE_AutonomousProxy)
-		{
-			LeaveScene();
-		}
-	}
 }
 
 FBoxSphereBounds UContextualAnimSceneActorComponent::CalcBounds(const FTransform& LocalToWorld) const
@@ -751,18 +750,15 @@ void UContextualAnimSceneActorComponent::LeaveScene()
 			*UEnum::GetValueAsString(TEXT("Engine.ENetRole"), GetOwner()->GetLocalRole()), *GetNameSafe(GetOwner()), *Bindings.GetRoleFromBinding(*Binding).ToString(),
 			Bindings.GetID(), Bindings.GetSectionIdx(), *GetNameSafe(Bindings.GetSceneAsset()));
 
-		const FContextualAnimTrack& AnimTrack = Bindings.GetAnimTrackFromBinding(*Binding);
-
 		if (UAnimInstance* AnimInstance = Binding->GetAnimInstance())
 		{
 			AnimInstance->OnMontageBlendingOut.RemoveDynamic(this, &UContextualAnimSceneActorComponent::OnMontageBlendingOut);
 
 			//@TODO: Add support for dynamic montage
-			UAnimMontage* AnimMontage = Cast<UAnimMontage>(AnimTrack.Animation);
-
-			if (AnimInstance->Montage_IsPlaying(AnimMontage))
+			const UAnimMontage* AnimMontage = AnimInstance->GetCurrentActiveMontage();
+			if (AnimMontage)
 			{
-				AnimInstance->Montage_Stop(AnimMontage->GetDefaultBlendOutTime(), AnimMontage);
+				AnimInstance->Montage_Stop(AnimMontage->GetDefaultBlendOutTime());
 			}
 		}
 
@@ -778,6 +774,7 @@ void UContextualAnimSceneActorComponent::LeaveScene()
 		// We might want to add a more robust mechanism to avoid overriding a request to disable collision that may have been set by another system
 		SetIgnoreCollisionWithOtherActors(false);
 
+		const FContextualAnimTrack& AnimTrack = Bindings.GetAnimTrackFromBinding(*Binding);
 		RestoreMovementState(AnimTrack.bRequireFlyingMode);
 
 		OnLeftSceneDelegate.Broadcast(this);
@@ -842,19 +839,15 @@ void UContextualAnimSceneActorComponent::OnMontageBlendingOut(UAnimMontage* Mont
 
 	if (const FContextualAnimSceneBinding* Binding = Bindings.FindBindingByActor(GetOwner()))
 	{
+		// Stop animation, restore state etc.
 		LeaveScene();
 
 		if (GetOwner()->HasAuthority())
 		{
-			// Rep empty bindings if we were the initiator of this interaction.
-			if (RepBindings.Bindings.IsValid())
-			{
-				RepBindings.Bindings.Clear();
-				RepBindings.ExternalWarpTargets.Empty();
-				MARK_PROPERTY_DIRTY_FROM_NAME(UContextualAnimSceneActorComponent, RepBindings, this);
-
-				GetOwner()->ForceNetUpdate();
-			}
+			RepTransitionSingleActorData.Reset();
+			RepTransitionSingleActorData.IncrementRepCounter();
+			MARK_PROPERTY_DIRTY_FROM_NAME(UContextualAnimSceneActorComponent, RepTransitionSingleActorData, this);
+			GetOwner()->ForceNetUpdate();
 		}
 	}
 }
diff --git a/Engine/Plugins/Experimental/Animation/ContextualAnimation/Source/ContextualAnimation/Public/ContextualAnimSceneActorComponent.h b/Engine/Plugins/Experimental/Animation/ContextualAnimation/Source/ContextualAnimation/Public/ContextualAnimSceneActorComponent.h
--- a/Engine/Plugins/Experimental/Animation/ContextualAnimation/Source/ContextualAnimation/Public/ContextualAnimSceneActorComponent.h
+++ b/Engine/Plugins/Experimental/Animation/ContextualAnimation/Source/ContextualAnimation/Public/ContextualAnimSceneActorComponent.h
@@ -100,6 +100,13 @@ struct FContextualAnimRepTransitionData : public FContextualAnimRepData
 
 	UPROPERTY()
 	TArray<FContextualAnimWarpTarget> ExternalWarpTargets;
+
+	void Reset()
+	{
+		SectionIdx = MAX_uint8;
+		AnimSetIdx = MAX_uint8;
+		ExternalWarpTargets.Reset();
+	}
 };
 
 UCLASS(meta = (BlueprintSpawnableComponent))
