Contextual Anim: Fixed bug causing the character to warp toward a previous interaction location on simulated proxies.

diff --git a/Engine/Plugins/Experimental/Animation/ContextualAnimation/Source/ContextualAnimation/Private/ContextualAnimTypes.cpp b/Engine/Plugins/Experimental/Animation/ContextualAnimation/Source/ContextualAnimation/Private/ContextualAnimTypes.cpp
--- a/Engine/Plugins/Experimental/Animation/ContextualAnimation/Source/ContextualAnimation/Private/ContextualAnimTypes.cpp
+++ b/Engine/Plugins/Experimental/Animation/ContextualAnimation/Source/ContextualAnimation/Private/ContextualAnimTypes.cpp
@@ -229,7 +229,7 @@ FVector FContextualAnimSceneBindingContext::GetVelocity() const
 
 UContextualAnimSceneActorComponent* FContextualAnimSceneBindingContext::GetSceneActorComponent() const
 {
-	if (!CachedSceneActorComp.IsValid())
+	if (!CachedSceneActorComp.IsValid() || CachedSceneActorComp->GetOwner() != GetActor())
 	{
 		if (Actor.IsValid())
 		{
@@ -242,7 +242,7 @@ UContextualAnimSceneActorComponent* FContextualAnimSceneBindingContext::GetScene
 
 UAnimInstance* FContextualAnimSceneBindingContext::GetAnimInstance() const
 {
-	if (!CachedAnimInstance.IsValid())
+	if (!CachedAnimInstance.IsValid() || CachedAnimInstance->GetOwningActor() != GetActor())
 	{
 		CachedAnimInstance = UContextualAnimUtilities::TryGetAnimInstance(GetActor());
 	}
@@ -252,7 +252,7 @@ UAnimInstance* FContextualAnimSceneBindingContext::GetAnimInstance() const
 
 USkeletalMeshComponent* FContextualAnimSceneBindingContext::GetSkeletalMeshComponent() const
 {
-	if (!CachedSkeletalMesh.IsValid())
+	if (!CachedSkeletalMesh.IsValid() || CachedSkeletalMesh->GetOwner() != GetActor())
 	{
 		CachedSkeletalMesh = UContextualAnimUtilities::TryGetSkeletalMeshComponent(GetActor());
 	}
