Contextual Anim: Fixed regression bug where simulated proxy can erroneously leave an interaction when playing two different interactions in a row.

diff --git a/Engine/Plugins/Experimental/Animation/ContextualAnimation/Source/ContextualAnimation/Private/ContextualAnimSceneActorComponent.cpp b/Engine/Plugins/Experimental/Animation/ContextualAnimation/Source/ContextualAnimation/Private/ContextualAnimSceneActorComponent.cpp
--- a/Engine/Plugins/Experimental/Animation/ContextualAnimation/Source/ContextualAnimation/Private/ContextualAnimSceneActorComponent.cpp
+++ b/Engine/Plugins/Experimental/Animation/ContextualAnimation/Source/ContextualAnimation/Private/ContextualAnimSceneActorComponent.cpp
@@ -319,6 +319,7 @@ bool UContextualAnimSceneActorComponent::TransitionContextualAnimScene(FName Sec
 
 			HandleTransitionEveryone(SectionIdx, 0, ExternalWarpTargets);
 
+			RepTransitionData.Id = Bindings.GetID();
 			RepTransitionData.SectionIdx = SectionIdx;
 			RepTransitionData.AnimSetIdx = 0;
 			RepTransitionData.ExternalWarpTargets = ExternalWarpTargets;
@@ -386,6 +387,7 @@ bool UContextualAnimSceneActorComponent::TransitionSingleActor(int32 SectionIdx,
 
 				AddOrUpdateWarpTargets(SectionIdx, AnimSetIdx, ExternalWarpTargets);
 
+				RepTransitionSingleActorData.Id = Bindings.GetID();
 				RepTransitionSingleActorData.SectionIdx = SectionIdx;
 				RepTransitionSingleActorData.AnimSetIdx = AnimSetIdx;
 				RepTransitionSingleActorData.ExternalWarpTargets = ExternalWarpTargets;
@@ -403,8 +405,8 @@ bool UContextualAnimSceneActorComponent::TransitionSingleActor(int32 SectionIdx,
 
 void UContextualAnimSceneActorComponent::OnRep_RepTransitionSingleActor()
 {
-	UE_LOG(LogContextualAnim, Verbose, TEXT("%-21s UContextualAnimSceneActorComponent::OnRep_RepTransitionSingleActor Owner: %s SectionIdx: %d AnimSetIdx: %d"),
-		*UEnum::GetValueAsString(TEXT("Engine.ENetRole"), GetOwner()->GetLocalRole()), *GetNameSafe(GetOwner()), RepTransitionSingleActorData.SectionIdx, RepTransitionSingleActorData.AnimSetIdx);
+	UE_LOG(LogContextualAnim, Verbose, TEXT("%-21s UContextualAnimSceneActorComponent::OnRep_RepTransitionSingleActor Owner: %s Id: %d SectionIdx: %d AnimSetIdx: %d"),
+		*UEnum::GetValueAsString(TEXT("Engine.ENetRole"), GetOwner()->GetLocalRole()), *GetNameSafe(GetOwner()), RepTransitionSingleActorData.Id, RepTransitionSingleActorData.SectionIdx, RepTransitionSingleActorData.AnimSetIdx);
 
 	if (const FContextualAnimSceneBinding* OwnerBinding = Bindings.FindBindingByActor(GetOwner()))
 	{
@@ -427,7 +429,7 @@ void UContextualAnimSceneActorComponent::OnRep_RepTransitionSingleActor()
 			// In this case we don't want to tell everyone else to also leave the scene since there is very common for the initiator, 
 			// specially if is player character, to end the animation earlier for responsiveness
 			// It is more likely this will do nothing since we listen to montage end also on Simulated Proxies to 'predict' the end of the interaction.
-			if (GetOwner()->GetLocalRole() != ROLE_AutonomousProxy)
+			if (RepTransitionSingleActorData.Id == Bindings.GetID() && GetOwner()->GetLocalRole() != ROLE_AutonomousProxy)
 			{
 				LeaveScene();
 			}
@@ -521,13 +523,18 @@ void UContextualAnimSceneActorComponent::EarlyOutContextualAnimScene()
 
 			if (Bindings.GetAnimTrackFromBinding(*Binding).Animation == ActiveMontage)
 			{
+				const uint8 BindingsId = Bindings.GetID();
+
 				// Stop animation.
 				LeaveScene();
 
 				// If we are on the server, rep the event to stop animation on simulated proxies
 				if (GetOwner()->HasAuthority())
 				{
-					RepTransitionSingleActorData.Reset();
+					RepTransitionSingleActorData.Id = BindingsId;
+					RepTransitionSingleActorData.SectionIdx = MAX_uint8;
+					RepTransitionSingleActorData.AnimSetIdx = MAX_uint8;
+					RepTransitionSingleActorData.ExternalWarpTargets.Reset();
 					RepTransitionSingleActorData.IncrementRepCounter();
 					MARK_PROPERTY_DIRTY_FROM_NAME(UContextualAnimSceneActorComponent, RepTransitionSingleActorData, this);
 					GetOwner()->ForceNetUpdate();
@@ -839,12 +846,17 @@ void UContextualAnimSceneActorComponent::OnMontageBlendingOut(UAnimMontage* Mont
 
 	if (const FContextualAnimSceneBinding* Binding = Bindings.FindBindingByActor(GetOwner()))
 	{
+		const uint8 BindingsId = Bindings.GetID();
+
 		// Stop animation, restore state etc.
 		LeaveScene();
 
 		if (GetOwner()->HasAuthority())
 		{
-			RepTransitionSingleActorData.Reset();
+			RepTransitionSingleActorData.Id = BindingsId;
+			RepTransitionSingleActorData.SectionIdx = MAX_uint8;
+			RepTransitionSingleActorData.AnimSetIdx = MAX_uint8;
+			RepTransitionSingleActorData.ExternalWarpTargets.Reset();
 			RepTransitionSingleActorData.IncrementRepCounter();
 			MARK_PROPERTY_DIRTY_FROM_NAME(UContextualAnimSceneActorComponent, RepTransitionSingleActorData, this);
 			GetOwner()->ForceNetUpdate();
diff --git a/Engine/Plugins/Experimental/Animation/ContextualAnimation/Source/ContextualAnimation/Public/ContextualAnimSceneActorComponent.h b/Engine/Plugins/Experimental/Animation/ContextualAnimation/Source/ContextualAnimation/Public/ContextualAnimSceneActorComponent.h
--- a/Engine/Plugins/Experimental/Animation/ContextualAnimation/Source/ContextualAnimation/Public/ContextualAnimSceneActorComponent.h
+++ b/Engine/Plugins/Experimental/Animation/ContextualAnimation/Source/ContextualAnimation/Public/ContextualAnimSceneActorComponent.h
@@ -92,6 +92,9 @@ struct FContextualAnimRepTransitionData : public FContextualAnimRepData
 {
 	GENERATED_BODY()
 
+	UPROPERTY()
+	uint8 Id = 0;
+
 	UPROPERTY()
 	uint8 SectionIdx = 0;
 
@@ -100,13 +103,6 @@ struct FContextualAnimRepTransitionData : public FContextualAnimRepData
 
 	UPROPERTY()
 	TArray<FContextualAnimWarpTarget> ExternalWarpTargets;
-
-	void Reset()
-	{
-		SectionIdx = MAX_uint8;
-		AnimSetIdx = MAX_uint8;
-		ExternalWarpTargets.Reset();
-	}
 };
 
 UCLASS(meta = (BlueprintSpawnableComponent))
