Contextual Anim: Members of FContextualAnimWarpTarget are now properly initialized.

diff --git a/Engine/Plugins/Experimental/Animation/ContextualAnimation/Source/ContextualAnimation/Public/ContextualAnimSceneActorComponent.h b/Engine/Plugins/Experimental/Animation/ContextualAnimation/Source/ContextualAnimation/Public/ContextualAnimSceneActorComponent.h
--- a/Engine/Plugins/Experimental/Animation/ContextualAnimation/Source/ContextualAnimation/Public/ContextualAnimSceneActorComponent.h
+++ b/Engine/Plugins/Experimental/Animation/ContextualAnimation/Source/ContextualAnimation/Public/ContextualAnimSceneActorComponent.h
@@ -30,16 +30,16 @@ struct FContextualAnimWarpTarget
 	}
 
 	UPROPERTY(EditAnywhere, Category = "Default")
-	FName Role;
+	FName Role = NAME_None;
 
 	UPROPERTY(EditAnywhere, Category = "Default")
-	FName TargetName;
+	FName TargetName = NAME_None;
 
 	UPROPERTY(EditAnywhere, Category = "Default")
-	FVector TargetLocation;
+	FVector TargetLocation = FVector::ZeroVector;
 
 	UPROPERTY(EditAnywhere, Category = "Default")
-	FQuat TargetRotation;
+	FQuat TargetRotation = FQuat::Identity;
 };
 
 /** Base struct for replicated data with a rep counter */
