Contextual Anim: Added a few safe checks with early out to prevent and track down a random crash happening in cooked builds.

diff --git a/Engine/Plugins/Experimental/Animation/ContextualAnimation/Source/ContextualAnimation/Private/ContextualAnimTypes.cpp b/Engine/Plugins/Experimental/Animation/ContextualAnimation/Source/ContextualAnimation/Private/ContextualAnimTypes.cpp
--- a/Engine/Plugins/Experimental/Animation/ContextualAnimation/Source/ContextualAnimation/Private/ContextualAnimTypes.cpp
+++ b/Engine/Plugins/Experimental/Animation/ContextualAnimation/Source/ContextualAnimation/Private/ContextualAnimTypes.cpp
@@ -388,26 +388,60 @@ bool FContextualAnimSceneBindings::BindActorToRole(AActor& ActorRef, FName Role)
 
 const FContextualAnimTrack& FContextualAnimSceneBindings::GetAnimTrackFromBinding(const FContextualAnimSceneBinding& Binding) const
 {
+	if(!IsValid())
+	{
+		UE_LOG(LogContextualAnim, Warning, TEXT("FContextualAnimSceneBindings::GetAnimTrackFromBinding Failed. Reason: Invalid Bindings. Bindings ID: %d Num: %d Scene Asset: %s. Binding Actor: %s"), 
+			GetID(), Num(), *GetNameSafe(GetSceneAsset()), *GetNameSafe(Binding.GetActor()));
+
+		return FContextualAnimTrack::EmptyTrack;
+	}
+
 	if (const FContextualAnimTrack* AnimTrack = GetSceneAsset()->GetAnimTrack(SectionIdx, AnimSetIdx, Binding.AnimTrackIdx))
 	{
 		return *AnimTrack;
 	}
 
+	UE_LOG(LogContextualAnim, Warning, TEXT("FContextualAnimSceneBindings::GetAnimTrackFromBinding Failed. Reason: Can't find AnimTrack for binding. Bindings ID: %d Num: %d Scene Asset: %s. Binding Actor: %s"),
+		GetID(), Num(), *GetNameSafe(GetSceneAsset()), *GetNameSafe(Binding.GetActor()));
 	return FContextualAnimTrack::EmptyTrack;
 }
 
 const FContextualAnimIKTargetDefContainer& FContextualAnimSceneBindings::GetIKTargetDefContainerFromBinding(const FContextualAnimSceneBinding& Binding) const
 {
+	if (!IsValid())
+	{
+		UE_LOG(LogContextualAnim, Warning, TEXT("FContextualAnimSceneBindings::GetIKTargetDefContainerFromBinding Failed. Reason: Invalid Bindings. Bindings ID: %d Num: %d Scene Asset: %s. Binding Actor: %s"),
+			GetID(), Num(), *GetNameSafe(GetSceneAsset()), *GetNameSafe(Binding.GetActor()));
+
+		return FContextualAnimIKTargetDefContainer::EmptyContainer;
+	}
+
 	return GetSceneAsset()->GetIKTargetDefsForRoleInSection(GetSectionIdx(), GetRoleFromBinding(Binding));
 }
 
 FTransform FContextualAnimSceneBindings::GetIKTargetTransformFromBinding(const FContextualAnimSceneBinding& Binding, const FName& TrackName, float Time) const
 {
+	if (!IsValid())
+	{
+		UE_LOG(LogContextualAnim, Warning, TEXT("FContextualAnimSceneBindings::GetIKTargetTransformFromBinding Failed. Reason: Invalid Bindings. Bindings ID: %d Num: %d Scene Asset: %s. Binding Actor: %s"),
+			GetID(), Num(), *GetNameSafe(GetSceneAsset()), *GetNameSafe(Binding.GetActor()));
+
+		return FTransform::Identity;
+	}
+
 	return GetSceneAsset()->GetIKTargetTransform(SectionIdx, AnimSetIdx, Binding.AnimTrackIdx, TrackName, Time);
 }
 
 FTransform FContextualAnimSceneBindings::GetAlignmentTransformFromBinding(const FContextualAnimSceneBinding& Binding, const FName& TrackName, float Time) const
 {
+	if (!IsValid())
+	{
+		UE_LOG(LogContextualAnim, Warning, TEXT("FContextualAnimSceneBindings::GetAlignmentTransformFromBinding Failed. Reason: Invalid Bindings. Bindings ID: %d Num: %d Scene Asset: %s. Binding Actor: %s"),
+			GetID(), Num(), *GetNameSafe(GetSceneAsset()), *GetNameSafe(Binding.GetActor()));
+
+		return FTransform::Identity;
+	}
+
 	return GetSceneAsset()->GetAlignmentTransform(SectionIdx, AnimSetIdx, Binding.AnimTrackIdx, TrackName, Time);
 }
 
