Contextual Anim Editor: Replaced checks with early outs and logs to prevent/track down a random crash that could happen in editor when editting Sections and/or Sets

diff --git a/Engine/Plugins/Experimental/Animation/ContextualAnimation/Source/ContextualAnimation/Private/ContextualAnimSceneAsset.cpp b/Engine/Plugins/Experimental/Animation/ContextualAnimation/Source/ContextualAnimation/Private/ContextualAnimSceneAsset.cpp
--- a/Engine/Plugins/Experimental/Animation/ContextualAnimation/Source/ContextualAnimation/Private/ContextualAnimSceneAsset.cpp
+++ b/Engine/Plugins/Experimental/Animation/ContextualAnimation/Source/ContextualAnimation/Private/ContextualAnimSceneAsset.cpp
@@ -698,12 +698,20 @@ FTransform UContextualAnimSceneAsset::GetAlignmentTransform(const FContextualAni
 	}
 
 	const FContextualAnimSceneSection* Section = GetSection(AnimTrack.SectionIdx);
-	check(Section);
+	if(Section == nullptr)
+	{
+		UE_LOG(LogContextualAnim, Warning, TEXT("UContextualAnimSceneAsset::GetAlignmentTransform Failed. Reason: Can't get AnimSceneSection. Asset: %s Section Idx: %d"), *GetNameSafe(this), AnimTrack.SectionIdx);
+		return FTransform::Identity;
+	}
 
 	if (Section->GetWarpPointDefinitions().IsValidIndex(WarpPointIdx))
 	{
 		const FContextualAnimSet* AnimSet = GetAnimSet(AnimTrack.SectionIdx, AnimTrack.AnimSetIdx);
-		check(AnimSet);
+		if (AnimSet == nullptr)
+		{
+			UE_LOG(LogContextualAnim, Warning, TEXT("UContextualAnimSceneAsset::GetAlignmentTransform Failed. Reason: Can't get AnimSet. Asset: %s Section Idx: %d Anim Set Idx: %d"), *GetNameSafe(this), AnimTrack.SectionIdx, AnimTrack.AnimSetIdx);
+			return FTransform::Identity;
+		}
 
 		const FName WarpTargetName = Section->GetWarpPointDefinitions()[WarpPointIdx].WarpTargetName;
 		if (const FTransform* WarpPointTransformPtr = AnimSet->WarpPoints.Find(WarpTargetName))
@@ -729,7 +737,11 @@ FTransform UContextualAnimSceneAsset::GetAlignmentTransform(const FContextualAni
 FTransform UContextualAnimSceneAsset::GetAlignmentTransform(const FContextualAnimTrack& AnimTrack, const FName& WarpPointName, float Time) const
 {
 	const FContextualAnimSceneSection* Section = GetSection(AnimTrack.SectionIdx);
-	check(Section);
+	if (Section == nullptr)
+	{
+		UE_LOG(LogContextualAnim, Warning, TEXT("UContextualAnimSceneAsset::GetAlignmentTransform Failed. Reason: Can't get AnimSceneSection. Asset: %s Section Idx: %d"), *GetNameSafe(this), AnimTrack.SectionIdx);
+		return FTransform::Identity;
+	}
 
 	for (int32 Idx = 0; Idx < Section->GetWarpPointDefinitions().Num(); Idx++)
 	{
