Fix server crash when more than 255 FContextualAnimSceneBindings are created.
UniqueIds would wrap around on UINT8_MAX to 0, which makes FContextualAnimSceneBindings::IsValid() return false for generated bindings, which could fail to generate data later on and crash.

diff --git a/Engine/Plugins/Experimental/Animation/ContextualAnimation/Source/ContextualAnimation/Private/ContextualAnimTypes.cpp b/Engine/Plugins/Experimental/Animation/ContextualAnimation/Source/ContextualAnimation/Private/ContextualAnimTypes.cpp
--- a/Engine/Plugins/Experimental/Animation/ContextualAnimation/Source/ContextualAnimation/Private/ContextualAnimTypes.cpp
+++ b/Engine/Plugins/Experimental/Animation/ContextualAnimation/Source/ContextualAnimation/Private/ContextualAnimTypes.cpp
@@ -322,7 +322,11 @@ FContextualAnimSceneBindings::FContextualAnimSceneBindings(const UContextualAnim
 void FContextualAnimSceneBindings::GenerateUniqueId()
 {
 	static uint8 IncrementID = 0;
-	IncrementID = IncrementID < UINT8_MAX ? IncrementID + 1 : 0;
+	if (IncrementID >= UINT8_MAX)
+	{
+		IncrementID = 0;
+	}
+	++IncrementID;
 	Id = IncrementID;
 }
 
