Contextual Anim: Removed binding actor from multiple output logs and replaced check with log and early out to prevent potential crash when bindings are invalid

diff --git a/Engine/Plugins/Experimental/Animation/ContextualAnimation/Source/ContextualAnimation/Private/ContextualAnimTypes.cpp b/Engine/Plugins/Experimental/Animation/ContextualAnimation/Source/ContextualAnimation/Private/ContextualAnimTypes.cpp
--- a/Engine/Plugins/Experimental/Animation/ContextualAnimation/Source/ContextualAnimation/Private/ContextualAnimTypes.cpp
+++ b/Engine/Plugins/Experimental/Animation/ContextualAnimation/Source/ContextualAnimation/Private/ContextualAnimTypes.cpp
@@ -394,8 +394,8 @@ const FContextualAnimTrack& FContextualAnimSceneBindings::GetAnimTrackFromBindin
 {
 	if(!IsValid())
 	{
-		UE_LOG(LogContextualAnim, Warning, TEXT("FContextualAnimSceneBindings::GetAnimTrackFromBinding Failed. Reason: Invalid Bindings. Bindings ID: %d Num: %d Scene Asset: %s. Binding Actor: %s"), 
-			GetID(), Num(), *GetNameSafe(GetSceneAsset()), *GetNameSafe(Binding.GetActor()));
+		UE_LOG(LogContextualAnim, Warning, TEXT("FContextualAnimSceneBindings::GetAnimTrackFromBinding Failed. Reason: Invalid Bindings. Bindings ID: %d Num: %d Scene Asset: %s"), 
+			GetID(), Num(), *GetNameSafe(GetSceneAsset()));
 
 		return FContextualAnimTrack::EmptyTrack;
 	}
@@ -405,8 +405,8 @@ const FContextualAnimTrack& FContextualAnimSceneBindings::GetAnimTrackFromBindin
 		return *AnimTrack;
 	}
 
-	UE_LOG(LogContextualAnim, Warning, TEXT("FContextualAnimSceneBindings::GetAnimTrackFromBinding Failed. Reason: Can't find AnimTrack for binding. Bindings ID: %d Num: %d Scene Asset: %s. Binding Actor: %s"),
-		GetID(), Num(), *GetNameSafe(GetSceneAsset()), *GetNameSafe(Binding.GetActor()));
+	UE_LOG(LogContextualAnim, Warning, TEXT("FContextualAnimSceneBindings::GetAnimTrackFromBinding Failed. Reason: Can't find AnimTrack for binding. Bindings ID: %d Num: %d Scene Asset: %s"),
+		GetID(), Num(), *GetNameSafe(GetSceneAsset()));
 	return FContextualAnimTrack::EmptyTrack;
 }
 
@@ -414,8 +414,8 @@ const FContextualAnimIKTargetDefContainer& FContextualAnimSceneBindings::GetIKTa
 {
 	if (!IsValid())
 	{
-		UE_LOG(LogContextualAnim, Warning, TEXT("FContextualAnimSceneBindings::GetIKTargetDefContainerFromBinding Failed. Reason: Invalid Bindings. Bindings ID: %d Num: %d Scene Asset: %s. Binding Actor: %s"),
-			GetID(), Num(), *GetNameSafe(GetSceneAsset()), *GetNameSafe(Binding.GetActor()));
+		UE_LOG(LogContextualAnim, Warning, TEXT("FContextualAnimSceneBindings::GetIKTargetDefContainerFromBinding Failed. Reason: Invalid Bindings. Bindings ID: %d Num: %d Scene Asset: %s"),
+			GetID(), Num(), *GetNameSafe(GetSceneAsset()));
 
 		return FContextualAnimIKTargetDefContainer::EmptyContainer;
 	}
@@ -427,8 +427,8 @@ FTransform FContextualAnimSceneBindings::GetIKTargetTransformFromBinding(const F
 {
 	if (!IsValid())
 	{
-		UE_LOG(LogContextualAnim, Warning, TEXT("FContextualAnimSceneBindings::GetIKTargetTransformFromBinding Failed. Reason: Invalid Bindings. Bindings ID: %d Num: %d Scene Asset: %s. Binding Actor: %s"),
-			GetID(), Num(), *GetNameSafe(GetSceneAsset()), *GetNameSafe(Binding.GetActor()));
+		UE_LOG(LogContextualAnim, Warning, TEXT("FContextualAnimSceneBindings::GetIKTargetTransformFromBinding Failed. Reason: Invalid Bindings. Bindings ID: %d Num: %d Scene Asset: %s"),
+			GetID(), Num(), *GetNameSafe(GetSceneAsset()));
 
 		return FTransform::Identity;
 	}
@@ -440,8 +440,8 @@ FTransform FContextualAnimSceneBindings::GetAlignmentTransformFromBinding(const
 {
 	if (!IsValid())
 	{
-		UE_LOG(LogContextualAnim, Warning, TEXT("FContextualAnimSceneBindings::GetAlignmentTransformFromBinding Failed. Reason: Invalid Bindings. Bindings ID: %d Num: %d Scene Asset: %s. Binding Actor: %s"),
-			GetID(), Num(), *GetNameSafe(GetSceneAsset()), *GetNameSafe(Binding.GetActor()));
+		UE_LOG(LogContextualAnim, Warning, TEXT("FContextualAnimSceneBindings::GetAlignmentTransformFromBinding Failed. Reason: Invalid Bindings. Bindings ID: %d Num: %d Scene Asset: %s"),
+			GetID(), Num(), *GetNameSafe(GetSceneAsset()));
 
 		return FTransform::Identity;
 	}
@@ -745,7 +745,12 @@ bool FContextualAnimSceneBindings::CalculateWarpPoint(const FContextualAnimWarpP
 
 void FContextualAnimSceneBindings::TransitionTo(int32 NewSectionIdx, int32 NewAnimSetIdx)
 {
-	check(IsValid());
+	if (!IsValid())
+	{
+		UE_LOG(LogContextualAnim, Warning, TEXT("FContextualAnimSceneBindings::TransitionTo NewSectionIdx: %d NewAnimSetIdx: %d Failed. Reason: Invalid Bindings. Bindings ID: %d Num: %d Scene Asset: %s."),
+			NewSectionIdx, NewAnimSetIdx, GetID(), Num(), *GetNameSafe(GetSceneAsset()));
+		return;
+	}
 
 	if (const FContextualAnimSet* AnimSet = SceneAsset->GetAnimSet(NewSectionIdx, NewAnimSetIdx))
 	{
