Contextual Anim: Replaced check with log and early out to prevent random crash when adding warp targets. Needs investigation.

diff --git a/Engine/Plugins/Experimental/Animation/ContextualAnimation/Source/ContextualAnimation/Private/ContextualAnimSceneActorComponent.cpp b/Engine/Plugins/Experimental/Animation/ContextualAnimation/Source/ContextualAnimation/Private/ContextualAnimSceneActorComponent.cpp
--- a/Engine/Plugins/Experimental/Animation/ContextualAnimation/Source/ContextualAnimation/Private/ContextualAnimSceneActorComponent.cpp
+++ b/Engine/Plugins/Experimental/Animation/ContextualAnimation/Source/ContextualAnimation/Private/ContextualAnimSceneActorComponent.cpp
@@ -124,10 +124,20 @@ void UContextualAnimSceneActorComponent::AddOrUpdateWarpTargets(int32 SectionIdx
 	if (const FContextualAnimSceneBinding* Binding = Bindings.FindBindingByActor(GetOwner()))
 	{
 		const UContextualAnimSceneAsset* Asset = Bindings.GetSceneAsset();
-		check(Asset);
+		if(Asset == nullptr)
+		{
+			UE_LOG(LogContextualAnim, Warning, TEXT("%-21s UContextualAnimSceneActorComponent::AddOrUpdateWarpTargets Invalid Scene Asset. Actor: %s Bindings Id: %d Bindings Num: %d SectionIdx: %d AnimSetIdx: %d"),
+				*UEnum::GetValueAsString(TEXT("Engine.ENetRole"), GetOwner()->GetLocalRole()), *GetNameSafe(GetOwner()), Bindings.GetID(), Bindings.Num(), SectionIdx, AnimSetIdx);
+			return;
+		}
 
 		const FContextualAnimSceneSection* Section = Asset->GetSection(SectionIdx);
-		check(Section);
+		if (Section == nullptr)
+		{
+			UE_LOG(LogContextualAnim, Warning, TEXT("%-21s UContextualAnimSceneActorComponent::AddOrUpdateWarpTargets Invalid Section. Actor: %s Bindings Id: %d Bindings Num: %d SectionIdx: %d AnimSetIdx: %d"),
+				*UEnum::GetValueAsString(TEXT("Engine.ENetRole"), GetOwner()->GetLocalRole()), *GetNameSafe(GetOwner()), Bindings.GetID(), Bindings.Num(), SectionIdx, AnimSetIdx);
+			return;
+		}
 
 		if(Section->GetWarpPointDefinitions().Num() > 0)
 		{
