Contextual Anim: Fixed standalone-only random bug where OnMontageBlendingOut will force the actor to erroneously leave the interaction when transitioning from one animation to another.

diff --git a/Engine/Plugins/Experimental/Animation/ContextualAnimation/Source/ContextualAnimation/Private/ContextualAnimSceneActorComponent.cpp b/Engine/Plugins/Experimental/Animation/ContextualAnimation/Source/ContextualAnimation/Private/ContextualAnimSceneActorComponent.cpp
--- a/Engine/Plugins/Experimental/Animation/ContextualAnimation/Source/ContextualAnimation/Private/ContextualAnimSceneActorComponent.cpp
+++ b/Engine/Plugins/Experimental/Animation/ContextualAnimation/Source/ContextualAnimation/Private/ContextualAnimSceneActorComponent.cpp
@@ -69,7 +69,17 @@ bool UContextualAnimSceneActorComponent::IsOwnerLocallyControlled() const
 
 void UContextualAnimSceneActorComponent::PlayAnimation_Internal(UAnimSequenceBase* Animation, float StartTime, bool bSyncPlaybackTime)
 {
-	TGuardValue<bool> UpdateGuard(bGuardAnimEvents, true);
+	// Replaced TGuardValue with this one frame delay because for some reason, apparently random (needs more investigation), in standalone OnMontageBlendingOut event is queued instead of triggered inline, 
+	// causing this guarding mechanism to fail because by the time the event triggers TGuardValue goes out of the scope
+	// Making our OnMontageBlendingOut to think the animation has been interrupted by an external system and forcing the actor to leave the interaction.
+	bGuardAnimEvents = true;
+	GetWorld()->GetTimerManager().SetTimerForNextTick([WeakThis = MakeWeakObjectPtr(this)]()
+	{
+		if (UContextualAnimSceneActorComponent* Comp = WeakThis.Get())
+		{
+			Comp->bGuardAnimEvents = false;
+		}
+	});
 
 	if (UAnimInstance* AnimInstance = UContextualAnimUtilities::TryGetAnimInstance(GetOwner()))
 	{
@@ -775,6 +785,7 @@ void UContextualAnimSceneActorComponent::LeaveScene()
 			const UAnimMontage* AnimMontage = AnimInstance->GetCurrentActiveMontage();
 			if (AnimMontage)
 			{
+				UE_LOG(LogContextualAnim, VeryVerbose, TEXT("\t\t Stopping animation (%s) from LeaveScene"), *GetNameSafe(AnimMontage));
 				AnimInstance->Montage_Stop(AnimMontage->GetDefaultBlendOutTime());
 			}
 		}
