Contextual Anim: Added extra logs to CheckConditions funtion to get details during the selection logic

diff --git a/Engine/Plugins/Experimental/Animation/ContextualAnimation/Source/ContextualAnimation/Private/ContextualAnimTypes.cpp b/Engine/Plugins/Experimental/Animation/ContextualAnimation/Source/ContextualAnimation/Private/ContextualAnimTypes.cpp
--- a/Engine/Plugins/Experimental/Animation/ContextualAnimation/Source/ContextualAnimation/Private/ContextualAnimTypes.cpp
+++ b/Engine/Plugins/Experimental/Animation/ContextualAnimation/Source/ContextualAnimation/Private/ContextualAnimTypes.cpp
@@ -511,6 +511,9 @@ bool FContextualAnimSceneBindings::CheckConditions(const UContextualAnimSceneAss
 		{
 			if (DoCheck(*PrimaryPtr, Pair.Value, RoleToBind) == false)
 			{
+				UE_LOG(LogContextualAnim, Verbose, TEXT("FContextualAnimSceneBindings::CheckConditions Failed. Reason: Can't find valid track for actor. SceneAsset: %s Actor: %s Role: %s SectionIdx: %d AnimSetIdx: %d"),
+					*GetNameSafe(&SceneAsset), *GetNameSafe(Pair.Value.GetActor()), *RoleToBind.ToString(), SectionIdx, AnimSetIdx);
+
 				return false;
 			}
 
@@ -521,9 +524,13 @@ bool FContextualAnimSceneBindings::CheckConditions(const UContextualAnimSceneAss
 	const int32 NumMandatoryRoles = SceneAsset.GetNumMandatoryRoles(SectionIdx, AnimSetIdx);
 	if (ValidEntries < NumMandatoryRoles)
 	{
+		UE_LOG(LogContextualAnim, Verbose, TEXT("FContextualAnimSceneBindings::CheckConditions Failed. Reason: Missing mandatory roles %d/%d. SceneAsset: %s SectionIdx: %d AnimSetIdx: %d"),
+			ValidEntries, NumMandatoryRoles, *GetNameSafe(&SceneAsset), SectionIdx, AnimSetIdx);
+
 		return false;
 	}
 
+	UE_LOG(LogContextualAnim, Verbose, TEXT("FContextualAnimSceneBindings::CheckConditions Success! SceneAsset: %s SectionIdx: %d AnimSetIdx: %d"), *GetNameSafe(&SceneAsset), SectionIdx, AnimSetIdx);
 	return true;
 }
 
