Contextual Anim Editor: Added no-animation section on the timeline + clicking on any section on the timeline will now select the corresponding anim track allowing the user to modify the properties for that specific track.

diff --git a/Engine/Plugins/Experimental/Animation/ContextualAnimation/Source/ContextualAnimationEditor/Private/ContextualAnimMovieSceneTrackEditor.cpp b/Engine/Plugins/Experimental/Animation/ContextualAnimation/Source/ContextualAnimationEditor/Private/ContextualAnimMovieSceneTrackEditor.cpp
--- a/Engine/Plugins/Experimental/Animation/ContextualAnimation/Source/ContextualAnimationEditor/Private/ContextualAnimMovieSceneTrackEditor.cpp
+++ b/Engine/Plugins/Experimental/Animation/ContextualAnimation/Source/ContextualAnimationEditor/Private/ContextualAnimMovieSceneTrackEditor.cpp
@@ -61,8 +61,14 @@ FText FContextualAnimSection::GetSectionTitle() const
 {
 	if (UContextualAnimMovieSceneSection* Section = Cast<UContextualAnimMovieSceneSection>(WeakSection.Get()))
 	{
-		return FText::FromString(GetNameSafe(Section->GetAnimTrack().Animation));
-		
+		if(Section->GetAnimTrack().Animation)
+		{
+			return FText::FromString(GetNameSafe(Section->GetAnimTrack().Animation));
+		}
+		else
+		{
+			return FText::FromString(TEXT("No Animation"));
+		}
 	}
 
 	return FText::FromString(TEXT("Animation_Name"));
diff --git a/Engine/Plugins/Experimental/Animation/ContextualAnimation/Source/ContextualAnimationEditor/Private/ContextualAnimViewModel.cpp b/Engine/Plugins/Experimental/Animation/ContextualAnimation/Source/ContextualAnimationEditor/Private/ContextualAnimViewModel.cpp
--- a/Engine/Plugins/Experimental/Animation/ContextualAnimation/Source/ContextualAnimationEditor/Private/ContextualAnimViewModel.cpp
+++ b/Engine/Plugins/Experimental/Animation/ContextualAnimation/Source/ContextualAnimationEditor/Private/ContextualAnimViewModel.cpp
@@ -142,6 +142,7 @@ void FContextualAnimViewModel::CreateSequencer()
 	Sequencer->OnGlobalTimeChanged().AddRaw(this, &FContextualAnimViewModel::SequencerTimeChanged);
 	Sequencer->OnPlayEvent().AddRaw(this, &FContextualAnimViewModel::SequencerPlayEvent);
 	Sequencer->OnStopEvent().AddRaw(this, &FContextualAnimViewModel::SequencerStopEvent);
+	Sequencer->GetSelectionChangedSections().AddRaw(this, &FContextualAnimViewModel::OnSequencerSelectionChangedSections);
 	Sequencer->SetPlaybackStatus(EMovieScenePlayerStatus::Stopped);
 }
 
@@ -397,6 +398,23 @@ void FContextualAnimViewModel::SetDefaultMode()
 	}
 
 	FContextualAnimSceneSection& ContextualAnimSection = SceneAsset->Sections[ActiveSectionIdx];
+
+	float EmptyAnimSectionLength = 0.f;
+	for (int32 AnimSetIdx = 0; AnimSetIdx < ContextualAnimSection.AnimSets.Num(); AnimSetIdx++)
+	{
+		const FContextualAnimSet& AnimSet = ContextualAnimSection.AnimSets[AnimSetIdx];
+		for (int32 AnimTrackIdx = 0; AnimTrackIdx < AnimSet.Tracks.Num(); AnimTrackIdx++)
+		{
+			const FContextualAnimTrack& AnimTrack = AnimSet.Tracks[AnimTrackIdx];
+			if(AnimTrack.Animation)
+			{
+				EmptyAnimSectionLength = FMath::Max(EmptyAnimSectionLength, AnimTrack.Animation->GetPlayLength());
+			}
+		}
+	}
+
+	EmptyAnimSectionLength = FMath::Max(1.f, EmptyAnimSectionLength);
+
 	for (int32 AnimSetIdx = 0; AnimSetIdx < ContextualAnimSection.AnimSets.Num(); AnimSetIdx++)
 	{
 		FContextualAnimSet& AnimSet = ContextualAnimSection.AnimSets[AnimSetIdx];
@@ -404,17 +422,16 @@ void FContextualAnimViewModel::SetDefaultMode()
 		{
 			const FContextualAnimTrack& AnimTrack = AnimSet.Tracks[AnimTrackIdx];
 
+			UContextualAnimMovieSceneTrack* MovieSceneTrack = FindTrackByRole(AnimTrack.Role);
+			check(MovieSceneTrack)
+
+			UContextualAnimMovieSceneSection* NewSection = nullptr;
 			UAnimSequenceBase* Animation = AnimTrack.Animation;
 			if (Animation)
 			{
-				UContextualAnimMovieSceneTrack* MovieSceneTrack = FindTrackByRole(AnimTrack.Role);
-				check(MovieSceneTrack)
-
-				UContextualAnimMovieSceneSection* NewSection = NewObject<UContextualAnimMovieSceneSection>(MovieSceneTrack, UContextualAnimMovieSceneSection::StaticClass(), NAME_None, RF_Transactional);
+				NewSection = NewObject<UContextualAnimMovieSceneSection>(MovieSceneTrack, UContextualAnimMovieSceneSection::StaticClass(), NAME_None, RF_Transactional);
 				check(NewSection);
 
-				NewSection->Initialize(ActiveSectionIdx, AnimSetIdx, AnimTrackIdx);
-
 				const float AnimLength = Animation->GetPlayLength();
 				const FFrameRate TickResolution = MovieSceneSequence->GetMovieScene()->GetTickResolution();
 				NewSection->SetRange(TRange<FFrameNumber>::Inclusive(0, (AnimLength * TickResolution).RoundToFrame()));
@@ -423,10 +440,22 @@ void FContextualAnimViewModel::SetDefaultMode()
 
 				const int32& ActiveSetIdx = ActiveAnimSetMap.FindOrAdd(ActiveSectionIdx);
 				NewSection->SetIsActive(AnimSetIdx == ActiveSetIdx);
+			}
+			else
+			{
+				NewSection = NewObject<UContextualAnimMovieSceneSection>(MovieSceneTrack, UContextualAnimMovieSceneSection::StaticClass(), NAME_None, RF_Transactional);
+				check(NewSection);
 
-				MovieSceneTrack->AddSection(*NewSection);
-				MovieSceneTrack->SetTrackRowDisplayName(FText::FromString(FString::Printf(TEXT("%d"), AnimSetIdx)), AnimSetIdx);
+				const FFrameRate TickResolution = MovieSceneSequence->GetMovieScene()->GetTickResolution();
+				NewSection->SetRange(TRange<FFrameNumber>::Inclusive(0, (EmptyAnimSectionLength * TickResolution).RoundToFrame()));
+				NewSection->SetRowIndex(AnimSetIdx);
+				NewSection->SetIsActive(false);
 			}
+
+			NewSection->Initialize(ActiveSectionIdx, AnimSetIdx, AnimTrackIdx);
+
+			MovieSceneTrack->AddSection(*NewSection);
+			MovieSceneTrack->SetTrackRowDisplayName(FText::FromString(FString::Printf(TEXT("%d"), AnimSetIdx)), AnimSetIdx);
 		}
 	}
 
@@ -941,6 +970,19 @@ void FContextualAnimViewModel::RefreshDetailsView()
 	DetailsView->ForceRefresh();
 }
 
+void FContextualAnimViewModel::OnSequencerSelectionChangedSections(TArray<UMovieSceneSection*> Sections)
+{
+	UE_LOG(LogContextualAnim, Log, TEXT("OnSequencerSelectionChangedSections Sections.Num: %d"), Sections.Num());
+
+	if (Sections.Num() > 0)
+	{
+		if (UContextualAnimMovieSceneSection* Section = Cast<UContextualAnimMovieSceneSection>(Sections[0]))
+		{
+			UpdateSelection(Section->GetSectionIdx(), Section->GetAnimSetIdx(), Section->GetAnimTrackIdx());
+		}
+	}
+}
+
 void FContextualAnimViewModel::UpdateSelection(const AActor* SelectedActor)
 {
 	const FContextualAnimSceneBinding* Binding = SceneBindings.FindBindingByActor(SelectedActor);
@@ -966,12 +1008,30 @@ void FContextualAnimViewModel::UpdateSelection(FName Role, int32 CriterionIdx, i
 	Sequencer->EmptySelection();
 
 	SelectionInfo.Role = Role;
+	SelectionInfo.SectionIdx = ActiveSectionIdx;
+	SelectionInfo.AnimSetIdx = ActiveAnimSetMap.FindRef(ActiveSectionIdx);
 	SelectionInfo.Criterion.Key = CriterionIdx;
 	SelectionInfo.Criterion.Value = CriterionDataIdx;
 
 	RefreshDetailsView();
 }
 
+void FContextualAnimViewModel::UpdateSelection(int32 SectionIdx, int32 AnimSetIdx, int32 AnimTrackIdx)
+{
+	if (const FContextualAnimTrack* AnimTrack = SceneAsset->GetAnimTrack(SectionIdx, AnimSetIdx, AnimTrackIdx))
+	{
+		Sequencer->EmptySelection();
+
+		SelectionInfo.Role = AnimTrack->Role;
+		SelectionInfo.SectionIdx = SectionIdx;
+		SelectionInfo.AnimSetIdx = AnimSetIdx;
+		SelectionInfo.Criterion.Key = INDEX_NONE;
+		SelectionInfo.Criterion.Value = INDEX_NONE;
+
+		RefreshDetailsView();
+	}
+}
+
 void FContextualAnimViewModel::ClearSelection()
 {
 	SelectionInfo.Reset();
@@ -992,8 +1052,8 @@ AActor* FContextualAnimViewModel::GetSelectedActor() const
 
 FContextualAnimTrack* FContextualAnimViewModel::GetSelectedAnimTrack() const
 {
-	FContextualAnimSceneBinding* Binding = GetSelectedBinding();
-	return Binding ? const_cast<FContextualAnimTrack*>(&SceneBindings.GetAnimTrackFromBinding(*Binding)) : nullptr;
+	const FContextualAnimTrack* AnimTrack = SceneAsset->GetAnimTrack(SelectionInfo.SectionIdx, SelectionInfo.AnimSetIdx, SelectionInfo.Role);
+	return AnimTrack ? const_cast<FContextualAnimTrack*>(AnimTrack) : nullptr;
 }
 
 UContextualAnimSelectionCriterion* FContextualAnimViewModel::GetSelectedSelectionCriterion() const
diff --git a/Engine/Plugins/Experimental/Animation/ContextualAnimation/Source/ContextualAnimationEditor/Public/ContextualAnimViewModel.h b/Engine/Plugins/Experimental/Animation/ContextualAnimation/Source/ContextualAnimationEditor/Public/ContextualAnimViewModel.h
--- a/Engine/Plugins/Experimental/Animation/ContextualAnimation/Source/ContextualAnimationEditor/Public/ContextualAnimViewModel.h
+++ b/Engine/Plugins/Experimental/Animation/ContextualAnimation/Source/ContextualAnimationEditor/Public/ContextualAnimViewModel.h
@@ -107,6 +107,7 @@ public:
 
 	void UpdateSelection(const AActor* SelectedActor);
 	void UpdateSelection(FName Role, int32 CriterionIdx = INDEX_NONE, int32 CriterionDataIdx = INDEX_NONE);
+	void UpdateSelection(int32 SectionIdx, int32 AnimSetIdx, int32 AnimTrackIdx);
 	void ClearSelection();
 	FContextualAnimSceneBinding* GetSelectedBinding() const;
 	FContextualAnimTrack* GetSelectedAnimTrack() const;
@@ -138,6 +139,10 @@ private:
 		/** Selected Role */
 		FName Role = NAME_None;
 
+		int32 SectionIdx = INDEX_NONE;
+
+		int32 AnimSetIdx = INDEX_NONE;
+
 		/** Key = Criterion Idx in the AnimTrack, Value = Data idx (e.g vertex idx) */
 		TPair<int32, int32> Criterion;
 
@@ -224,6 +229,8 @@ private:
 
 	void SequencerStopEvent();
 
+	void OnSequencerSelectionChangedSections(TArray<UMovieSceneSection*> Sections);
+
 	void OnAnimNotifyChanged(UAnimSequenceBase* Animation);
 
 	void CreateSequencer();
