Contextual Anim: Added version of TransitionContextualAnimScene function that receives an specific anim set via param.

diff --git a/Engine/Plugins/Experimental/Animation/ContextualAnimation/Source/ContextualAnimation/Private/ContextualAnimSceneActorComponent.cpp b/Engine/Plugins/Experimental/Animation/ContextualAnimation/Source/ContextualAnimation/Private/ContextualAnimSceneActorComponent.cpp
--- a/Engine/Plugins/Experimental/Animation/ContextualAnimation/Source/ContextualAnimation/Private/ContextualAnimSceneActorComponent.cpp
+++ b/Engine/Plugins/Experimental/Animation/ContextualAnimation/Source/ContextualAnimation/Private/ContextualAnimSceneActorComponent.cpp
@@ -311,7 +311,7 @@ void UContextualAnimSceneActorComponent::OnRep_LateJoinData()
 	HandleLateJoin(RepLateJoinData.Actor, RepLateJoinData.Role, RepLateJoinData.ExternalWarpTargets);
 }
 
-bool UContextualAnimSceneActorComponent::TransitionContextualAnimScene(FName SectionName, const TArray<FContextualAnimWarpTarget>& ExternalWarpTargets)
+bool UContextualAnimSceneActorComponent::TransitionContextualAnimScene(FName SectionName, int32 AnimSetIdx, const TArray<FContextualAnimWarpTarget>& ExternalWarpTargets)
 {
 	if (!GetOwner()->HasAuthority())
 	{
@@ -326,7 +326,7 @@ bool UContextualAnimSceneActorComponent::TransitionContextualAnimScene(FName Sec
 		{
 			if (UContextualAnimSceneActorComponent* Comp = Leader->GetSceneActorComponent())
 			{
-				return Comp->TransitionContextualAnimScene(SectionName, ExternalWarpTargets);
+				return Comp->TransitionContextualAnimScene(SectionName, AnimSetIdx, ExternalWarpTargets);
 			}
 		}
 	}
@@ -339,11 +339,11 @@ bool UContextualAnimSceneActorComponent::TransitionContextualAnimScene(FName Sec
 			UE_LOG(LogContextualAnim, Log, TEXT("%-21s UContextualAnimSceneActorComponent::TransitionTo Actor: %s SectionName: %s"),
 				*UEnum::GetValueAsString(TEXT("Engine.ENetRole"), GetOwner()->GetLocalRole()), *GetNameSafe(GetOwner()), *SectionName.ToString());
 
-			HandleTransitionEveryone(SectionIdx, 0, ExternalWarpTargets);
+			HandleTransitionEveryone(SectionIdx, AnimSetIdx, ExternalWarpTargets);
 
 			RepTransitionData.Id = Bindings.GetID();
 			RepTransitionData.SectionIdx = SectionIdx;
-			RepTransitionData.AnimSetIdx = 0;
+			RepTransitionData.AnimSetIdx = AnimSetIdx;
 			RepTransitionData.ExternalWarpTargets = ExternalWarpTargets;
 			RepTransitionData.IncrementRepCounter();
 			MARK_PROPERTY_DIRTY_FROM_NAME(UContextualAnimSceneActorComponent, RepTransitionData, this);
@@ -356,6 +356,12 @@ bool UContextualAnimSceneActorComponent::TransitionContextualAnimScene(FName Sec
 	return false;
 }
 
+bool UContextualAnimSceneActorComponent::TransitionContextualAnimScene(FName SectionName, const TArray<FContextualAnimWarpTarget>& ExternalWarpTargets)
+{
+	//@TODO: Run selection criteria to find best AnimSet
+	return TransitionContextualAnimScene(SectionName, 0, ExternalWarpTargets);
+}
+
 void UContextualAnimSceneActorComponent::HandleTransitionEveryone(int32 NewSectionIdx, int32 NewAnimSetIdx, const TArray<FContextualAnimWarpTarget>& ExternalWarpTargets)
 {
 	// Update Bindings internal data and play new animation for the leader first
diff --git a/Engine/Plugins/Experimental/Animation/ContextualAnimation/Source/ContextualAnimation/Public/ContextualAnimSceneActorComponent.h b/Engine/Plugins/Experimental/Animation/ContextualAnimation/Source/ContextualAnimation/Public/ContextualAnimSceneActorComponent.h
--- a/Engine/Plugins/Experimental/Animation/ContextualAnimation/Source/ContextualAnimation/Public/ContextualAnimSceneActorComponent.h
+++ b/Engine/Plugins/Experimental/Animation/ContextualAnimation/Source/ContextualAnimation/Public/ContextualAnimSceneActorComponent.h
@@ -167,6 +167,7 @@ public:
 	bool StartContextualAnimScene(const FContextualAnimSceneBindings& InBindings, const TArray<FContextualAnimWarpTarget>& ExternalWarpTargets);
 	bool LateJoinContextualAnimScene(AActor* Actor, FName Role, const TArray<FContextualAnimWarpTarget>& ExternalWarpTargets);
 	bool TransitionContextualAnimScene(FName SectionName, const TArray<FContextualAnimWarpTarget>& ExternalWarpTargets);
+	bool TransitionContextualAnimScene(FName SectionName, int32 AnimSetIdx, const TArray<FContextualAnimWarpTarget>& ExternalWarpTargets);
 	bool TransitionSingleActor(int32 SectionIdx, int32 AnimSetIdx, const TArray<FContextualAnimWarpTarget>& ExternalWarpTargets);
 
 	void EarlyOutContextualAnimScene();
