Contextual Anim: Potential fix for random bug only in cooked build that happens when gc runs and the current scene asset gets collected leaving the interaction in the simulasted proxy in an invalid state.

diff --git a/Engine/Plugins/Experimental/Animation/ContextualAnimation/Source/ContextualAnimation/Private/ContextualAnimSceneActorComponent.cpp b/Engine/Plugins/Experimental/Animation/ContextualAnimation/Source/ContextualAnimation/Private/ContextualAnimSceneActorComponent.cpp
--- a/Engine/Plugins/Experimental/Animation/ContextualAnimation/Source/ContextualAnimation/Private/ContextualAnimSceneActorComponent.cpp
+++ b/Engine/Plugins/Experimental/Animation/ContextualAnimation/Source/ContextualAnimation/Private/ContextualAnimSceneActorComponent.cpp
@@ -47,6 +47,14 @@ UContextualAnimSceneActorComponent::UContextualAnimSceneActorComponent(const FOb
 	SetCollisionEnabled(ECollisionEnabled::NoCollision);
 }
 
+void UContextualAnimSceneActorComponent::AddReferencedObjects(UObject* InThis, FReferenceCollector& Collector)
+{
+	UContextualAnimSceneActorComponent* This = CastChecked<UContextualAnimSceneActorComponent>(InThis);
+	This->Bindings.AddReferencedObjects(Collector);
+
+	Super::AddReferencedObjects(This, Collector);
+}
+
 void UContextualAnimSceneActorComponent::GetLifetimeReplicatedProps(TArray< FLifetimeProperty >& OutLifetimeProps) const
 {
 	Super::GetLifetimeReplicatedProps(OutLifetimeProps);
diff --git a/Engine/Plugins/Experimental/Animation/ContextualAnimation/Source/ContextualAnimation/Private/ContextualAnimTypes.cpp b/Engine/Plugins/Experimental/Animation/ContextualAnimation/Source/ContextualAnimation/Private/ContextualAnimTypes.cpp
--- a/Engine/Plugins/Experimental/Animation/ContextualAnimation/Source/ContextualAnimation/Private/ContextualAnimTypes.cpp
+++ b/Engine/Plugins/Experimental/Animation/ContextualAnimation/Source/ContextualAnimation/Private/ContextualAnimTypes.cpp
@@ -330,15 +330,23 @@ void FContextualAnimSceneBindings::GenerateUniqueId()
 	Id = IncrementID;
 }
 
+void FContextualAnimSceneBindings::AddReferencedObjects(FReferenceCollector& Collector)
+{
+	if (SceneAsset)
+	{
+		Collector.AddReferencedObject(SceneAsset);
+	}
+}
+
 bool FContextualAnimSceneBindings::IsValid() const
 {
-	return Id != 0 && SceneAsset.IsValid() && SceneAsset->HasValidData() && Num() > 0;
+	return Id != 0 && SceneAsset && SceneAsset->HasValidData() && Num() > 0;
 }
 
 void FContextualAnimSceneBindings::Reset()
 {
 	Id = 0;
-	SceneAsset.Reset();
+	SceneAsset = nullptr;
 	SectionIdx = INDEX_NONE;
 	AnimSetIdx = INDEX_NONE;
 	Data.Reset();
@@ -381,7 +389,7 @@ bool FContextualAnimSceneBindings::BindActorToRole(AActor& ActorRef, FName Role)
 	if(AnimTrackPtr == nullptr)
 	{
 		UE_LOG(LogContextualAnim, Warning, TEXT("FContextualAnimSceneBindings::BindActorToRole. Failed to bind Actor: '%s' to Role: '%s'. Reason: Can't find valid AnimTrack for it. SceneAsset: %s SectionIdx: %d AnimSetIdx: %d"),
-			*GetNameSafe(&ActorRef), *Role.ToString(), *GetNameSafe(SceneAsset.Get()), SectionIdx, AnimSetIdx);
+			*GetNameSafe(&ActorRef), *Role.ToString(), *GetNameSafe(SceneAsset), SectionIdx, AnimSetIdx);
 		return false;
 	}
 
diff --git a/Engine/Plugins/Experimental/Animation/ContextualAnimation/Source/ContextualAnimation/Public/ContextualAnimSceneActorComponent.h b/Engine/Plugins/Experimental/Animation/ContextualAnimation/Source/ContextualAnimation/Public/ContextualAnimSceneActorComponent.h
--- a/Engine/Plugins/Experimental/Animation/ContextualAnimation/Source/ContextualAnimation/Public/ContextualAnimSceneActorComponent.h
+++ b/Engine/Plugins/Experimental/Animation/ContextualAnimation/Source/ContextualAnimation/Public/ContextualAnimSceneActorComponent.h
@@ -128,6 +128,7 @@ public:
 
 	UContextualAnimSceneActorComponent(const FObjectInitializer& ObjectInitializer);
 
+	static void AddReferencedObjects(UObject* InThis, FReferenceCollector& Collector);
 	virtual void GetLifetimeReplicatedProps(TArray< FLifetimeProperty >& OutLifetimeProps) const;
 	virtual void OnRegister() override;
 	virtual void OnUnregister() override;
diff --git a/Engine/Plugins/Experimental/Animation/ContextualAnimation/Source/ContextualAnimation/Public/ContextualAnimTypes.h b/Engine/Plugins/Experimental/Animation/ContextualAnimation/Source/ContextualAnimation/Public/ContextualAnimTypes.h
--- a/Engine/Plugins/Experimental/Animation/ContextualAnimation/Source/ContextualAnimation/Public/ContextualAnimTypes.h
+++ b/Engine/Plugins/Experimental/Animation/ContextualAnimation/Source/ContextualAnimation/Public/ContextualAnimTypes.h
@@ -429,8 +429,10 @@ struct CONTEXTUALANIMATION_API FContextualAnimSceneBindings
 		return Role != NAME_None ? Data.FindByPredicate([this, &Role](const FContextualAnimSceneBinding& Item) { return GetAnimTrackFromBinding(Item).Role == Role; }) : nullptr;
 	}
 
+	void AddReferencedObjects(FReferenceCollector& Collector);
+
 	FORCEINLINE uint8 GetID() const { return Id; }
-	FORCEINLINE const UContextualAnimSceneAsset* GetSceneAsset() const { return SceneAsset.Get(); }
+	FORCEINLINE const UContextualAnimSceneAsset* GetSceneAsset() const { return SceneAsset; }
 	FORCEINLINE int32 GetSectionIdx() const { return SectionIdx; }
 	FORCEINLINE int32 GetAnimSetIdx() const { return AnimSetIdx; }
 	FORCEINLINE int32 Num() const { return Data.Num(); }
@@ -482,7 +484,7 @@ private:
 	uint8 Id = 0;
 
 	UPROPERTY()
-	TWeakObjectPtr<const UContextualAnimSceneAsset> SceneAsset = nullptr;
+	TObjectPtr<const UContextualAnimSceneAsset> SceneAsset = nullptr;
 
 	UPROPERTY()
 	int32 SectionIdx = INDEX_NONE;
