Temporary workaround to have CA work with ClusterUnion from physics

diff --git a/Engine/Plugins/Experimental/Animation/ContextualAnimation/Source/ContextualAnimation/Private/ContextualAnimSceneActorComponent.cpp b/Engine/Plugins/Experimental/Animation/ContextualAnimation/Source/ContextualAnimation/Private/ContextualAnimSceneActorComponent.cpp
--- a/Engine/Plugins/Experimental/Animation/ContextualAnimation/Source/ContextualAnimation/Private/ContextualAnimSceneActorComponent.cpp
+++ b/Engine/Plugins/Experimental/Animation/ContextualAnimation/Source/ContextualAnimation/Private/ContextualAnimSceneActorComponent.cpp
@@ -16,6 +16,8 @@
 #include "Rig/IKRigDataTypes.h"
 #include "Net/Core/PushModel/PushModel.h"
 #include "Physics/PhysicsInterfaceTypes.h"
+#include "PhysicsEngine/ClusterUnionComponent.h"
+#include "PhysicsEngine/ClusterUnionReplicatedProxyComponent.h"
 #include "PrimitiveSceneProxy.h"
 #include "PrimitiveViewRelevance.h"
 #include "SceneManagement.h"
@@ -676,10 +678,28 @@ void UContextualAnimSceneActorComponent::SetIgnoreCollisionWithOtherActors(bool
 	for (const FContextualAnimSceneBinding& Binding : Bindings)
 	{
 		AActor* OtherActor = Binding.GetActor();
-		if (OtherActor != OwnerActor)
+		if (OtherActor && OtherActor != OwnerActor)
 		{
 			if (UPrimitiveComponent* RootPrimitiveComponent = Cast<UPrimitiveComponent>(OwnerActor->GetRootComponent()))
 			{
+				// Temporary hack to fix issue with ClusterUnion not being transparently integrated in the physics systems
+				{
+					TArray<UClusterUnionReplicatedProxyComponent*> ClusterUnionProxyComponents;
+					OtherActor->GetComponents<UClusterUnionReplicatedProxyComponent>(ClusterUnionProxyComponents);
+
+					// We can assume that all the proxy components will point to the same cluster union.
+					if (!ClusterUnionProxyComponents.IsEmpty() && !ClusterUnionProxyComponents[0]->IsPendingDeletion())
+					{
+						if (const UClusterUnionComponent* const ClusterUnionComponent = ClusterUnionProxyComponents[0]->GetParentClusterUnionComponent())
+						{
+							if (AActor* const Owner = ClusterUnionComponent->GetOwner())
+							{
+								RootPrimitiveComponent->IgnoreActorWhenMoving(Owner, bValue);
+							}
+						}
+					}
+				}
+			
 				RootPrimitiveComponent->IgnoreActorWhenMoving(OtherActor, bValue);
 			}
 		}
