Avoid CA crash with bad networking conditions by checking for a nullptr Actor before deferencing it

diff --git a/Engine/Plugins/Experimental/Animation/ContextualAnimation/Source/ContextualAnimation/Private/ContextualAnimSceneActorComponent.cpp b/Engine/Plugins/Experimental/Animation/ContextualAnimation/Source/ContextualAnimation/Private/ContextualAnimSceneActorComponent.cpp
--- a/Engine/Plugins/Experimental/Animation/ContextualAnimation/Source/ContextualAnimation/Private/ContextualAnimSceneActorComponent.cpp
+++ b/Engine/Plugins/Experimental/Animation/ContextualAnimation/Source/ContextualAnimation/Private/ContextualAnimSceneActorComponent.cpp
@@ -239,7 +239,7 @@ bool UContextualAnimSceneActorComponent::LateJoinContextualAnimScene(AActor* Act
 
 bool UContextualAnimSceneActorComponent::HandleLateJoin(AActor* Actor, FName Role, const TArray<FContextualAnimWarpTarget>& ExternalWarpTargets)
 {
-	if (!Bindings.BindActorToRole(*Actor, Role))
+	if (!IsValid(Actor) || !Bindings.BindActorToRole(*Actor, Role))
 	{
 		UE_LOG(LogContextualAnim, Warning, TEXT("%-21s UContextualAnimSceneActorComponent::HandleLateJoin Failed. Reason: Adding %s to the bindings for role: %s failed!"),
 			*UEnum::GetValueAsString(TEXT("Engine.ENetRole"), GetOwner()->GetLocalRole()), *GetNameSafe(Actor), *Role.ToString());
