Contextual Anim: Selection criteria is now used during transitions to select the best set.

diff --git a/Engine/Plugins/Experimental/Animation/ContextualAnimation/Source/ContextualAnimation/Private/ContextualAnimSceneActorComponent.cpp b/Engine/Plugins/Experimental/Animation/ContextualAnimation/Source/ContextualAnimation/Private/ContextualAnimSceneActorComponent.cpp
--- a/Engine/Plugins/Experimental/Animation/ContextualAnimation/Source/ContextualAnimation/Private/ContextualAnimSceneActorComponent.cpp
+++ b/Engine/Plugins/Experimental/Animation/ContextualAnimation/Source/ContextualAnimation/Private/ContextualAnimSceneActorComponent.cpp
@@ -368,8 +368,31 @@ bool UContextualAnimSceneActorComponent::TransitionContextualAnimScene(FName Sec
 
 bool UContextualAnimSceneActorComponent::TransitionContextualAnimScene(FName SectionName, const TArray<FContextualAnimWarpTarget>& ExternalWarpTargets)
 {
-	//@TODO: Run selection criteria to find best AnimSet
-	return TransitionContextualAnimScene(SectionName, 0, ExternalWarpTargets);
+	if (!Bindings.IsValid())
+	{
+		UE_LOG(LogContextualAnim, Warning, TEXT("%-21s UContextualAnimSceneActorComponent::TransitionContextualAnimScene Invalid Bindings"), *UEnum::GetValueAsString(TEXT("Engine.ENetRole"), GetOwner()->GetLocalRole()));
+		return false;
+	}
+
+	const int32 SectionIdx = Bindings.GetSceneAsset()->GetSectionIndex(SectionName);
+	if (SectionIdx == INDEX_NONE)
+	{
+		UE_LOG(LogContextualAnim, Log, TEXT("%-21s UContextualAnimSceneActorComponent::TransitionContextualAnimScene. Invalid SectionName. Actor: %s SectionName: %s"),
+			*UEnum::GetValueAsString(TEXT("Engine.ENetRole"), GetOwner()->GetLocalRole()), *GetNameSafe(GetOwner()), *SectionName.ToString());
+
+		return false;
+	}
+
+	const int32 AnimSetIdx = Bindings.FindAnimSetForTransitionTo(SectionIdx);
+	if (AnimSetIdx == INDEX_NONE)
+	{
+		UE_LOG(LogContextualAnim, Log, TEXT("%-21s UContextualAnimSceneActorComponent::TransitionContextualAnimScene. Can't find AnimSet. Actor: %s SectionName: %s"),
+			*UEnum::GetValueAsString(TEXT("Engine.ENetRole"), GetOwner()->GetLocalRole()), *GetNameSafe(GetOwner()), *SectionName.ToString());
+
+		return false;
+	}
+
+	return TransitionContextualAnimScene(SectionName, AnimSetIdx, ExternalWarpTargets);
 }
 
 void UContextualAnimSceneActorComponent::HandleTransitionEveryone(int32 NewSectionIdx, int32 NewAnimSetIdx, const TArray<FContextualAnimWarpTarget>& ExternalWarpTargets)
@@ -404,6 +427,26 @@ void UContextualAnimSceneActorComponent::HandleTransitionSelf(int32 NewSectionId
 	AddOrUpdateWarpTargets(NewSectionIdx, NewAnimSetIdx, ExternalWarpTargets);
 }
 
+bool UContextualAnimSceneActorComponent::TransitionSingleActor(int32 SectionIdx, const TArray<FContextualAnimWarpTarget>& ExternalWarpTargets)
+{
+	if (!Bindings.IsValid())
+	{
+		UE_LOG(LogContextualAnim, Warning, TEXT("%-21s UContextualAnimSceneActorComponent::TransitionSingleActor Invalid Bindings"), *UEnum::GetValueAsString(TEXT("Engine.ENetRole"), GetOwner()->GetLocalRole()));
+		return false;
+	}
+
+	const int32 AnimSetIdx = Bindings.FindAnimSetForTransitionTo(SectionIdx);
+	if (AnimSetIdx == INDEX_NONE)
+	{
+		UE_LOG(LogContextualAnim, Log, TEXT("%-21s UContextualAnimSceneActorComponent::TransitionSingleActor. Can't find AnimSet. Actor: %s SectionIdx: %d"),
+			*UEnum::GetValueAsString(TEXT("Engine.ENetRole"), GetOwner()->GetLocalRole()), *GetNameSafe(GetOwner()), SectionIdx);
+
+		return false;
+	}
+
+	return TransitionSingleActor(SectionIdx, AnimSetIdx, ExternalWarpTargets);
+}
+
 bool UContextualAnimSceneActorComponent::TransitionSingleActor(int32 SectionIdx, int32 AnimSetIdx, const TArray<FContextualAnimWarpTarget>& ExternalWarpTargets)
 {
 	if (!GetOwner()->HasAuthority())
diff --git a/Engine/Plugins/Experimental/Animation/ContextualAnimation/Source/ContextualAnimation/Private/ContextualAnimTypes.cpp b/Engine/Plugins/Experimental/Animation/ContextualAnimation/Source/ContextualAnimation/Private/ContextualAnimTypes.cpp
--- a/Engine/Plugins/Experimental/Animation/ContextualAnimation/Source/ContextualAnimation/Private/ContextualAnimTypes.cpp
+++ b/Engine/Plugins/Experimental/Animation/ContextualAnimation/Source/ContextualAnimation/Private/ContextualAnimTypes.cpp
@@ -591,10 +591,9 @@ bool FContextualAnimSceneBindings::TryCreateBindings(const UContextualAnimSceneA
 	return TryCreateBindings(SceneAsset, SectionIdx, AnimSetIdx, Params, OutBindings);
 }
 
-bool FContextualAnimSceneBindings::TryCreateBindings(const UContextualAnimSceneAsset& SceneAsset, int32 SectionIdx, const TMap<FName, FContextualAnimSceneBindingContext>& Params, FContextualAnimSceneBindings& OutBindings)
+int32 FContextualAnimSceneBindings::FindAnimSet(const UContextualAnimSceneAsset& SceneAsset, int32 SectionIdx, const TMap<FName, FContextualAnimSceneBindingContext>& Params)
 {
 	check(SceneAsset.HasValidData());
-	OutBindings.Reset();
 
 	int32 AnimSetIdxSelected = INDEX_NONE;
 
@@ -610,7 +609,7 @@ bool FContextualAnimSceneBindings::TryCreateBindings(const UContextualAnimSceneA
 	{
 		const FContextualAnimSceneSection* Section = SceneAsset.GetSection(SectionIdx);
 		TArray<TTuple<int32, float>, TInlineAllocator<5>> ValidSets; // 0: AnimSetIdx, 1: AnimSetRandomWeight
-		
+
 		float TotalWeight = 0;
 		for (int32 AnimSetIdx = 0; AnimSetIdx < NumSets; AnimSetIdx++)
 		{
@@ -634,6 +633,16 @@ bool FContextualAnimSceneBindings::TryCreateBindings(const UContextualAnimSceneA
 		}
 	}
 
+	UE_CLOG(AnimSetIdxSelected == INDEX_NONE, LogContextualAnim, Warning, TEXT("FContextualAnimSceneBindings::FindAnimSet. Can't find AnimSet"));
+	return AnimSetIdxSelected;
+}
+
+bool FContextualAnimSceneBindings::TryCreateBindings(const UContextualAnimSceneAsset& SceneAsset, int32 SectionIdx, const TMap<FName, FContextualAnimSceneBindingContext>& Params, FContextualAnimSceneBindings& OutBindings)
+{
+	check(SceneAsset.HasValidData());
+	OutBindings.Reset();
+
+	const int32 AnimSetIdxSelected = FindAnimSet(SceneAsset, SectionIdx, Params);
 	if (AnimSetIdxSelected != INDEX_NONE)
 	{
 		OutBindings = FContextualAnimSceneBindings(SceneAsset, SectionIdx, AnimSetIdxSelected);
@@ -652,6 +661,32 @@ bool FContextualAnimSceneBindings::TryCreateBindings(const UContextualAnimSceneA
 	return false;
 }
 
+int32 FContextualAnimSceneBindings::FindAnimSetForTransitionTo(int32 NewSectionIdx) const
+{
+	if (!IsValid())
+	{
+		UE_LOG(LogContextualAnim, Warning, TEXT("FContextualAnimSceneBindings::FindAnimSetForTransition Invalid Bindings"));
+		return INDEX_NONE;
+	}
+
+	const FContextualAnimSceneSection* Section = GetSceneAsset()->GetSection(NewSectionIdx);
+	if (Section == nullptr)
+	{
+		UE_LOG(LogContextualAnim, Warning, TEXT("FContextualAnimSceneBindings::FindAnimSetForTransition. Invalid NewSectionIdx %d"), NewSectionIdx);
+		return INDEX_NONE;
+	}
+
+	TMap<FName, FContextualAnimSceneBindingContext> Params;
+	Params.Reserve(Num());
+	for (const FContextualAnimSceneBinding& Binding : Data)
+	{
+		Params.Add(GetRoleFromBinding(Binding), Binding.GetContext());
+	}
+
+	const int32 AnimSetIdxSelected = FContextualAnimSceneBindings::FindAnimSet(*GetSceneAsset(), NewSectionIdx, Params);
+	return AnimSetIdxSelected;
+}
+
 bool FContextualAnimSceneBindings::TryCreateBindings(const UContextualAnimSceneAsset& SceneAsset, int32 SectionIdx, const FContextualAnimSceneBindingContext& Primary, const FContextualAnimSceneBindingContext& Secondary, FContextualAnimSceneBindings& OutBindings)
 {
 	check(SceneAsset.HasValidData());
diff --git a/Engine/Plugins/Experimental/Animation/ContextualAnimation/Source/ContextualAnimation/Public/ContextualAnimSceneActorComponent.h b/Engine/Plugins/Experimental/Animation/ContextualAnimation/Source/ContextualAnimation/Public/ContextualAnimSceneActorComponent.h
--- a/Engine/Plugins/Experimental/Animation/ContextualAnimation/Source/ContextualAnimation/Public/ContextualAnimSceneActorComponent.h
+++ b/Engine/Plugins/Experimental/Animation/ContextualAnimation/Source/ContextualAnimation/Public/ContextualAnimSceneActorComponent.h
@@ -169,6 +169,7 @@ public:
 	bool LateJoinContextualAnimScene(AActor* Actor, FName Role, const TArray<FContextualAnimWarpTarget>& ExternalWarpTargets);
 	bool TransitionContextualAnimScene(FName SectionName, const TArray<FContextualAnimWarpTarget>& ExternalWarpTargets);
 	bool TransitionContextualAnimScene(FName SectionName, int32 AnimSetIdx, const TArray<FContextualAnimWarpTarget>& ExternalWarpTargets);
+	bool TransitionSingleActor(int32 SectionIdx, const TArray<FContextualAnimWarpTarget>& ExternalWarpTargets);
 	bool TransitionSingleActor(int32 SectionIdx, int32 AnimSetIdx, const TArray<FContextualAnimWarpTarget>& ExternalWarpTargets);
 
 	void EarlyOutContextualAnimScene();
diff --git a/Engine/Plugins/Experimental/Animation/ContextualAnimation/Source/ContextualAnimation/Public/ContextualAnimTypes.h b/Engine/Plugins/Experimental/Animation/ContextualAnimation/Source/ContextualAnimation/Public/ContextualAnimTypes.h
--- a/Engine/Plugins/Experimental/Animation/ContextualAnimation/Source/ContextualAnimation/Public/ContextualAnimTypes.h
+++ b/Engine/Plugins/Experimental/Animation/ContextualAnimation/Source/ContextualAnimation/Public/ContextualAnimTypes.h
@@ -449,6 +449,7 @@ struct CONTEXTUALANIMATION_API FContextualAnimSceneBindings
 	static bool TryCreateBindings(const UContextualAnimSceneAsset& SceneAsset, int32 SectionIdx, const TMap<FName, FContextualAnimSceneBindingContext>& Params, FContextualAnimSceneBindings& OutBindings);
 	static bool TryCreateBindings(const UContextualAnimSceneAsset& SceneAsset, int32 SectionIdx, int32 AnimSetIdx, const FContextualAnimSceneBindingContext& Primary, const FContextualAnimSceneBindingContext& Secondary, FContextualAnimSceneBindings& OutBindings);
 	static bool TryCreateBindings(const UContextualAnimSceneAsset& SceneAsset, int32 SectionIdx, const FContextualAnimSceneBindingContext& Primary, const FContextualAnimSceneBindingContext& Secondary, FContextualAnimSceneBindings& OutBindings);
+	static int32 FindAnimSet(const UContextualAnimSceneAsset& SceneAsset, int32 SectionIdx, const TMap<FName, FContextualAnimSceneBindingContext>& Params);
 
 	bool BindActorToRole(AActor& ActorRef, FName Role);
 
@@ -474,6 +475,7 @@ struct CONTEXTUALANIMATION_API FContextualAnimSceneBindings
 
 	void GenerateUniqueId();
 
+	int32 FindAnimSetForTransitionTo(int32 NewSectionIdx) const;
 	void TransitionTo(int32 NewSectionIdx, int32 NewAnimSetIdx);
 
 private:
