Contextual Anim: Fixed bug where simulated proxies would receive and react to previous interaction events after becoming net relevant.

diff --git a/Engine/Plugins/Experimental/Animation/ContextualAnimation/Source/ContextualAnimation/Private/ContextualAnimSceneActorComponent.cpp b/Engine/Plugins/Experimental/Animation/ContextualAnimation/Source/ContextualAnimation/Private/ContextualAnimSceneActorComponent.cpp
--- a/Engine/Plugins/Experimental/Animation/ContextualAnimation/Source/ContextualAnimation/Private/ContextualAnimSceneActorComponent.cpp
+++ b/Engine/Plugins/Experimental/Animation/ContextualAnimation/Source/ContextualAnimation/Private/ContextualAnimSceneActorComponent.cpp
@@ -35,7 +35,11 @@ TAutoConsoleVariable<float> CVarContextualAnimIKDrawDebugLifetime(TEXT("a.Contex
 void FContextualAnimRepData::IncrementRepCounter()
 {
 	static uint8 Counter = 0;
-	Counter = Counter < UINT8_MAX ? Counter + 1 : 0;
+	if (Counter >= UINT8_MAX)
+	{
+		Counter = 0;
+	}
+	++Counter;
 	RepCounter = Counter;
 }
 
@@ -308,15 +312,20 @@ void UContextualAnimSceneActorComponent::OnRep_LateJoinData()
 {
 	// This is received by the leader of the interaction on every remote client
 
+	UE_LOG(LogContextualAnim, Verbose, TEXT("%-21s UContextualAnimSceneActorComponent::OnRep_LateJoinData Owner: %s Bindings Id: %d Section: %d Asset: %s. Requester: %s Role: %s RepCounter: %d"),
+		*UEnum::GetValueAsString(TEXT("Engine.ENetRole"), GetOwner()->GetLocalRole()), *GetNameSafe(GetOwner()), Bindings.GetID(), Bindings.GetSectionIdx(), *GetNameSafe(Bindings.GetSceneAsset()), *GetNameSafe(RepLateJoinData.Actor), *RepLateJoinData.Role.ToString(), RepLateJoinData.RepCounter);
+
+	if (!RepLateJoinData.IsValid())
+	{
+		return;
+	}
+
 	if (!Bindings.IsValid())
 	{
 		UE_LOG(LogContextualAnim, Warning, TEXT("%-21s UContextualAnimSceneActorComponent::OnRep_LateJoinData Invalid Bindings"), *UEnum::GetValueAsString(TEXT("Engine.ENetRole"), GetOwner()->GetLocalRole()));
 		return;
 	}
 
-	UE_LOG(LogContextualAnim, Verbose, TEXT("%-21s UContextualAnimSceneActorComponent::OnRep_LateJoinData Owner: %s Bindings Id: %d Section: %d Asset: %s. Requester: %s Role: %s"),
-		*UEnum::GetValueAsString(TEXT("Engine.ENetRole"), GetOwner()->GetLocalRole()), *GetNameSafe(GetOwner()), Bindings.GetID(), Bindings.GetSectionIdx(), *GetNameSafe(Bindings.GetSceneAsset()), *GetNameSafe(RepLateJoinData.Actor), *RepLateJoinData.Role.ToString());
-
 	// Play animation and set state on this new actor that is joining us and update bindings for everyone else
 	HandleLateJoin(RepLateJoinData.Actor, RepLateJoinData.Role, RepLateJoinData.ExternalWarpTargets);
 }
@@ -486,8 +495,13 @@ bool UContextualAnimSceneActorComponent::TransitionSingleActor(int32 SectionIdx,
 
 void UContextualAnimSceneActorComponent::OnRep_RepTransitionSingleActor()
 {
-	UE_LOG(LogContextualAnim, Verbose, TEXT("%-21s UContextualAnimSceneActorComponent::OnRep_RepTransitionSingleActor Owner: %s Id: %d SectionIdx: %d AnimSetIdx: %d"),
-		*UEnum::GetValueAsString(TEXT("Engine.ENetRole"), GetOwner()->GetLocalRole()), *GetNameSafe(GetOwner()), RepTransitionSingleActorData.Id, RepTransitionSingleActorData.SectionIdx, RepTransitionSingleActorData.AnimSetIdx);
+	UE_LOG(LogContextualAnim, Verbose, TEXT("%-21s UContextualAnimSceneActorComponent::OnRep_RepTransitionSingleActor Owner: %s Id: %d RepCounter: %d SectionIdx: %d AnimSetIdx: %d"),
+		*UEnum::GetValueAsString(TEXT("Engine.ENetRole"), GetOwner()->GetLocalRole()), *GetNameSafe(GetOwner()), RepTransitionSingleActorData.Id, RepTransitionSingleActorData.RepCounter, RepTransitionSingleActorData.SectionIdx, RepTransitionSingleActorData.AnimSetIdx);
+
+	if (!RepTransitionSingleActorData.IsValid())
+	{
+		return;
+	}
 
 	if (const FContextualAnimSceneBinding* OwnerBinding = Bindings.FindBindingByActor(GetOwner()))
 	{
@@ -563,6 +577,8 @@ bool UContextualAnimSceneActorComponent::StartContextualAnimScene(const FContext
 
 			RepBindings.Bindings = InBindings;
 			RepBindings.ExternalWarpTargets = ExternalWarpTargets;
+			RepBindings.IncrementRepCounter();
+
 			MARK_PROPERTY_DIRTY_FROM_NAME(UContextualAnimSceneActorComponent, RepBindings, this);
 			GetOwner()->ForceNetUpdate();
 
@@ -617,7 +633,16 @@ void UContextualAnimSceneActorComponent::EarlyOutContextualAnimScene()
 					RepTransitionSingleActorData.AnimSetIdx = MAX_uint8;
 					RepTransitionSingleActorData.ExternalWarpTargets.Reset();
 					RepTransitionSingleActorData.IncrementRepCounter();
+
+					RepLateJoinData.Reset();
+					RepTransitionData.Reset();
+					RepBindings.Reset();
+
 					MARK_PROPERTY_DIRTY_FROM_NAME(UContextualAnimSceneActorComponent, RepTransitionSingleActorData, this);
+					MARK_PROPERTY_DIRTY_FROM_NAME(UContextualAnimSceneActorComponent, RepLateJoinData, this);
+					MARK_PROPERTY_DIRTY_FROM_NAME(UContextualAnimSceneActorComponent, RepTransitionData, this);
+					MARK_PROPERTY_DIRTY_FROM_NAME(UContextualAnimSceneActorComponent, RepBindings, this);
+
 					GetOwner()->ForceNetUpdate();
 				}
 				// If local player, tell the server to stop the animation too
@@ -642,16 +667,27 @@ bool UContextualAnimSceneActorComponent::ServerEarlyOutContextualAnimScene_Valid
 
 void UContextualAnimSceneActorComponent::OnRep_TransitionData()
 {
-	UE_LOG(LogContextualAnim, Verbose, TEXT("%-21s UContextualAnimSceneActorComponent::OnRep_TransitionData Actor: %s"),
-		*UEnum::GetValueAsString(TEXT("Engine.ENetRole"), GetOwner()->GetLocalRole()), *GetNameSafe(GetOwner()));
+	UE_LOG(LogContextualAnim, Verbose, TEXT("%-21s UContextualAnimSceneActorComponent::OnRep_TransitionData Actor: %s SectionIdx: %d AnimsetIdx: %d RepCounter: %d"),
+		*UEnum::GetValueAsString(TEXT("Engine.ENetRole"), GetOwner()->GetLocalRole()), *GetNameSafe(GetOwner()),
+		RepTransitionData.SectionIdx, RepTransitionData.AnimSetIdx, RepTransitionData.RepCounter);
+
+	if (!RepTransitionData.IsValid())
+	{
+		return;
+	}
 
 	HandleTransitionEveryone(RepTransitionData.SectionIdx, RepTransitionData.AnimSetIdx, RepTransitionData.ExternalWarpTargets);
 }
 
 void UContextualAnimSceneActorComponent::OnRep_Bindings()
 {
-	UE_LOG(LogContextualAnim, Verbose, TEXT("%-21s UContextualAnimSceneActorComponent::OnRep_Bindings Actor: %s Rep Bindings Id: %d Num: %d Current Bindings Id: %d Num: %d"),
-		*UEnum::GetValueAsString(TEXT("Engine.ENetRole"), GetOwner()->GetLocalRole()), *GetNameSafe(GetOwner()), RepBindings.Bindings.GetID(), RepBindings.Bindings.Num(), Bindings.GetID(), Bindings.Num());
+	UE_LOG(LogContextualAnim, Verbose, TEXT("%-21s UContextualAnimSceneActorComponent::OnRep_Bindings Actor: %s Rep Bindings Id: %d RepCounter: %d Num: %d Current Bindings Id: %d Num: %d"),
+		*UEnum::GetValueAsString(TEXT("Engine.ENetRole"), GetOwner()->GetLocalRole()), *GetNameSafe(GetOwner()), RepBindings.Bindings.GetID(), RepBindings.RepCounter, RepBindings.Bindings.Num(), Bindings.GetID(), Bindings.Num());
+
+	if (!RepBindings.IsValid())
+	{
+		return;
+	}
 
 	// The owner of this component started an interaction on the server
 	if (RepBindings.Bindings.IsValid())
@@ -958,7 +994,16 @@ void UContextualAnimSceneActorComponent::OnMontageBlendingOut(UAnimMontage* Mont
 			RepTransitionSingleActorData.AnimSetIdx = MAX_uint8;
 			RepTransitionSingleActorData.ExternalWarpTargets.Reset();
 			RepTransitionSingleActorData.IncrementRepCounter();
+
+			RepLateJoinData.Reset();
+			RepTransitionData.Reset();
+			RepBindings.Reset();
+
 			MARK_PROPERTY_DIRTY_FROM_NAME(UContextualAnimSceneActorComponent, RepTransitionSingleActorData, this);
+			MARK_PROPERTY_DIRTY_FROM_NAME(UContextualAnimSceneActorComponent, RepLateJoinData, this);
+			MARK_PROPERTY_DIRTY_FROM_NAME(UContextualAnimSceneActorComponent, RepTransitionData, this);
+			MARK_PROPERTY_DIRTY_FROM_NAME(UContextualAnimSceneActorComponent, RepBindings, this);
+
 			GetOwner()->ForceNetUpdate();
 		}
 	}
diff --git a/Engine/Plugins/Experimental/Animation/ContextualAnimation/Source/ContextualAnimation/Public/ContextualAnimSceneActorComponent.h b/Engine/Plugins/Experimental/Animation/ContextualAnimation/Source/ContextualAnimation/Public/ContextualAnimSceneActorComponent.h
--- a/Engine/Plugins/Experimental/Animation/ContextualAnimation/Source/ContextualAnimation/Public/ContextualAnimSceneActorComponent.h
+++ b/Engine/Plugins/Experimental/Animation/ContextualAnimation/Source/ContextualAnimation/Public/ContextualAnimSceneActorComponent.h
@@ -53,6 +53,7 @@ struct FContextualAnimRepData
 	uint8 RepCounter = 0;
 
 	void IncrementRepCounter();
+	bool IsValid() const { return RepCounter != 0; }
 };
 
 /** Used to replicate start/stop contextual anim events */
@@ -66,6 +67,13 @@ struct FContextualAnimRepBindingsData : public FContextualAnimRepData
 
 	UPROPERTY()
 	TArray<FContextualAnimWarpTarget> ExternalWarpTargets;
+
+	void Reset()
+	{
+		RepCounter = 0;
+		Bindings.Reset();
+		ExternalWarpTargets.Reset();
+	}
 };
 
 /** Used to replicate a late join event */
@@ -84,6 +92,14 @@ struct FContextualAnimRepLateJoinData : public FContextualAnimRepData
 
 	UPROPERTY()
 	TArray<FContextualAnimWarpTarget> ExternalWarpTargets;
+
+	void Reset()
+	{
+		RepCounter = 0;
+		Actor = nullptr;
+		Role = NAME_None;
+		ExternalWarpTargets.Reset();
+	}
 };
 
 /** Used to transition events */
@@ -103,6 +119,15 @@ struct FContextualAnimRepTransitionData : public FContextualAnimRepData
 
 	UPROPERTY()
 	TArray<FContextualAnimWarpTarget> ExternalWarpTargets;
+
+	void Reset()
+	{
+		RepCounter = 0;
+		Id = 0;
+		SectionIdx = 0;
+		AnimSetIdx = 0;
+		ExternalWarpTargets.Reset();
+	}
 };
 
 UCLASS(meta = (BlueprintSpawnableComponent))
