Contextual Anim: Removed temp workaround related to ClusterUnion submitted in 25048264

diff --git a/Engine/Plugins/Experimental/Animation/ContextualAnimation/Source/ContextualAnimation/Private/ContextualAnimSceneActorComponent.cpp b/Engine/Plugins/Experimental/Animation/ContextualAnimation/Source/ContextualAnimation/Private/ContextualAnimSceneActorComponent.cpp
--- a/Engine/Plugins/Experimental/Animation/ContextualAnimation/Source/ContextualAnimation/Private/ContextualAnimSceneActorComponent.cpp
+++ b/Engine/Plugins/Experimental/Animation/ContextualAnimation/Source/ContextualAnimation/Private/ContextualAnimSceneActorComponent.cpp
@@ -16,8 +16,6 @@
 #include "Rig/IKRigDataTypes.h"
 #include "Net/Core/PushModel/PushModel.h"
 #include "Physics/PhysicsInterfaceTypes.h"
-#include "PhysicsEngine/ClusterUnionComponent.h"
-#include "PhysicsEngine/ClusterUnionReplicatedProxyComponent.h"
 #include "PrimitiveSceneProxy.h"
 #include "PrimitiveViewRelevance.h"
 #include "SceneManagement.h"
@@ -761,24 +759,6 @@ void UContextualAnimSceneActorComponent::SetIgnoreCollisionWithOtherActors(bool
 		{
 			if (UPrimitiveComponent* RootPrimitiveComponent = Cast<UPrimitiveComponent>(OwnerActor->GetRootComponent()))
 			{
-				// Temporary hack to fix issue with ClusterUnion not being transparently integrated in the physics systems
-				{
-					TArray<UClusterUnionReplicatedProxyComponent*> ClusterUnionProxyComponents;
-					OtherActor->GetComponents<UClusterUnionReplicatedProxyComponent>(ClusterUnionProxyComponents);
-
-					// We can assume that all the proxy components will point to the same cluster union.
-					if (!ClusterUnionProxyComponents.IsEmpty() && !ClusterUnionProxyComponents[0]->IsPendingDeletion())
-					{
-						if (const UClusterUnionComponent* const ClusterUnionComponent = ClusterUnionProxyComponents[0]->GetParentClusterUnionComponent())
-						{
-							if (AActor* const Owner = ClusterUnionComponent->GetOwner())
-							{
-								RootPrimitiveComponent->IgnoreActorWhenMoving(Owner, bValue);
-							}
-						}
-					}
-				}
-			
 				RootPrimitiveComponent->IgnoreActorWhenMoving(OtherActor, bValue);
 			}
 		}
