Contextual Anim: JoinScene and LeaveScene are now called from OnRep event even for AutonomousProxy so player interactions can be started on the server. In the future this is gonna be parametrized so we can have both behaviors.

diff --git a/Engine/Plugins/Experimental/Animation/ContextualAnimation/Source/ContextualAnimation/Private/ContextualAnimSceneActorComponent.cpp b/Engine/Plugins/Experimental/Animation/ContextualAnimation/Source/ContextualAnimation/Private/ContextualAnimSceneActorComponent.cpp
--- a/Engine/Plugins/Experimental/Animation/ContextualAnimation/Source/ContextualAnimation/Private/ContextualAnimSceneActorComponent.cpp
+++ b/Engine/Plugins/Experimental/Animation/ContextualAnimation/Source/ContextualAnimation/Private/ContextualAnimSceneActorComponent.cpp
@@ -522,7 +522,7 @@ void UContextualAnimSceneActorComponent::OnRep_RepTransitionSingleActor()
 			// In this case we don't want to tell everyone else to also leave the scene since there is very common for the initiator, 
 			// specially if is player character, to end the animation earlier for responsiveness
 			// It is more likely this will do nothing since we listen to montage end also on Simulated Proxies to 'predict' the end of the interaction.
-			if (RepTransitionSingleActorData.Id == Bindings.GetID() && GetOwner()->GetLocalRole() != ROLE_AutonomousProxy)
+			if (RepTransitionSingleActorData.Id == Bindings.GetID())
 			{
 				LeaveScene();
 			}
@@ -694,10 +694,7 @@ void UContextualAnimSceneActorComponent::OnRep_Bindings()
 		if (ensureAlways(OwnerBinding))
 		{
 			// Join the scene (start playing animation, etc.)
-			if (GetOwner()->GetLocalRole() != ROLE_AutonomousProxy)
-			{
-				JoinScene(RepBindings.Bindings, RepBindings.ExternalWarpTargets);
-			}
+			JoinScene(RepBindings.Bindings, RepBindings.ExternalWarpTargets);
 
 			// RepBindings is only replicated from the initiator of the action.
 			// So now we have to tell everyone else involved in the interaction to join us
