ContextualAnim: Added PlayMontageNotifyBegin delegate to notify external actors when PlayMontageNotify triggers during an interaction.

diff --git a/Engine/Plugins/Experimental/Animation/ContextualAnimation/Source/ContextualAnimation/Private/ContextualAnimSceneActorComponent.cpp b/Engine/Plugins/Experimental/Animation/ContextualAnimation/Source/ContextualAnimation/Private/ContextualAnimSceneActorComponent.cpp
--- a/Engine/Plugins/Experimental/Animation/ContextualAnimation/Source/ContextualAnimation/Private/ContextualAnimSceneActorComponent.cpp
+++ b/Engine/Plugins/Experimental/Animation/ContextualAnimation/Source/ContextualAnimation/Private/ContextualAnimSceneActorComponent.cpp
@@ -105,6 +105,8 @@ void UContextualAnimSceneActorComponent::PlayAnimation_Internal(UAnimSequenceBas
 		AnimInstance->Montage_Play(AnimMontage, 1.f, EMontagePlayReturnType::MontageLength, StartTime);
 
 		AnimInstance->OnMontageBlendingOut.AddUniqueDynamic(this, &UContextualAnimSceneActorComponent::OnMontageBlendingOut);
+		AnimInstance->OnPlayMontageNotifyBegin.AddUniqueDynamic(this, &UContextualAnimSceneActorComponent::OnPlayMontageNotifyBegin);
+		
 
 		if (bSyncPlaybackTime)
 		{
@@ -872,6 +874,7 @@ void UContextualAnimSceneActorComponent::LeaveScene()
 		if (UAnimInstance* AnimInstance = Binding->GetAnimInstance())
 		{
 			AnimInstance->OnMontageBlendingOut.RemoveDynamic(this, &UContextualAnimSceneActorComponent::OnMontageBlendingOut);
+			AnimInstance->OnPlayMontageNotifyBegin.RemoveDynamic(this, &UContextualAnimSceneActorComponent::OnPlayMontageNotifyBegin);
 
 			//@TODO: Add support for dynamic montage
 			const UAnimMontage* AnimMontage = AnimInstance->GetCurrentActiveMontage();
@@ -986,6 +989,21 @@ void UContextualAnimSceneActorComponent::OnMontageBlendingOut(UAnimMontage* Mont
 	}
 }
 
+
+void UContextualAnimSceneActorComponent::OnPlayMontageNotifyBegin(FName NotifyName, const FBranchingPointNotifyPayload& BranchingPointNotifyPayload)
+{
+	if (bGuardAnimEvents)
+	{
+		return;
+	}
+
+	UE_LOG(LogContextualAnim, Verbose, TEXT("%-21s UContextualAnimSceneActorComponent::OnNotifyBeginReceived Actor: %s Animation: %s NotifyName"),
+		*UEnum::GetValueAsString(TEXT("Engine.ENetRole"), GetOwner()->GetLocalRole()), *GetNameSafe(GetOwner()), *GetNameSafe(BranchingPointNotifyPayload.SequenceAsset), *NotifyName.ToString());
+
+	OnPlayMontageNotifyBeginDelegate.Broadcast(this, NotifyName);
+
+}
+
 void UContextualAnimSceneActorComponent::OnTickPose(class USkinnedMeshComponent* SkinnedMeshComponent, float DeltaTime, bool bNeedsValidRootMotion)
 {
 	//@TODO: Check for LOD to prevent this update if the actor is too far away
diff --git a/Engine/Plugins/Experimental/Animation/ContextualAnimation/Source/ContextualAnimation/Public/ContextualAnimSceneActorComponent.h b/Engine/Plugins/Experimental/Animation/ContextualAnimation/Source/ContextualAnimation/Public/ContextualAnimSceneActorComponent.h
--- a/Engine/Plugins/Experimental/Animation/ContextualAnimation/Source/ContextualAnimation/Public/ContextualAnimSceneActorComponent.h
+++ b/Engine/Plugins/Experimental/Animation/ContextualAnimation/Source/ContextualAnimation/Public/ContextualAnimSceneActorComponent.h
@@ -14,6 +14,7 @@ class UAnimMontage;
 class UAnimSequenceBase;
 
 DECLARE_DYNAMIC_MULTICAST_DELEGATE_OneParam(FContextualAnimSceneActorCompDelegate, class UContextualAnimSceneActorComponent*, SceneActorComponent);
+DECLARE_DYNAMIC_MULTICAST_DELEGATE_TwoParams(FContextualAnimPlayMontageNotifyBeginDelegate, class UContextualAnimSceneActorComponent*, SceneActorComponent, FName, NotifyName);
 
 USTRUCT(BlueprintType)
 struct FContextualAnimWarpTarget
@@ -145,6 +146,9 @@ public:
 	UPROPERTY(BlueprintAssignable, Category = "Events")
 	FContextualAnimSceneActorCompDelegate OnLeftSceneDelegate;
 
+	UPROPERTY(BlueprintAssignable, Category = "Events")
+	FContextualAnimPlayMontageNotifyBeginDelegate OnPlayMontageNotifyBeginDelegate;
+
 	UPROPERTY(EditAnywhere, BlueprintReadWrite, Category = "Config")
 	TObjectPtr<class UContextualAnimSceneAsset> SceneAsset;
 
@@ -268,6 +272,9 @@ protected:
 	UFUNCTION()
 	void OnMontageBlendingOut(UAnimMontage* Montage, bool bInterrupted);
 
+	UFUNCTION()
+	void OnPlayMontageNotifyBegin(FName NotifyName, const FBranchingPointNotifyPayload& BranchingPointNotifyPayload);
+
 	// @TODO: These two functions are going to replace OnJoinedScene and OnLeftScene
 	// main different is that these new functions are taking care of animation playback too
 
