Contextual Anim: Added accessor functions for binding context's external gameplay tags

diff --git a/Engine/Plugins/Experimental/Animation/ContextualAnimation/Source/ContextualAnimation/Private/ContextualAnimTypes.cpp b/Engine/Plugins/Experimental/Animation/ContextualAnimation/Source/ContextualAnimation/Private/ContextualAnimTypes.cpp
--- a/Engine/Plugins/Experimental/Animation/ContextualAnimation/Source/ContextualAnimation/Private/ContextualAnimTypes.cpp
+++ b/Engine/Plugins/Experimental/Animation/ContextualAnimation/Source/ContextualAnimation/Private/ContextualAnimTypes.cpp
@@ -196,6 +196,26 @@ void FContextualAnimSceneBindingContext::SetExternalTransform(const FTransform&
 	ExternalTransform = InTransform;
 }
 
+void FContextualAnimSceneBindingContext::AddGameplayTag(const FGameplayTag& Tag)
+{
+	ExternalGameplayTags.AddTag(Tag);
+}
+
+bool FContextualAnimSceneBindingContext::HasMatchingGameplayTag(const FGameplayTag& TagToCheck) const
+{
+	return ExternalGameplayTags.HasTag(TagToCheck);
+}
+
+bool FContextualAnimSceneBindingContext::HasAllMatchingGameplayTags(const FGameplayTagContainer& TagContainer) const
+{
+	return ExternalGameplayTags.HasAll(TagContainer);
+}
+
+bool FContextualAnimSceneBindingContext::HasAnyMatchingGameplayTags(const FGameplayTagContainer& TagContainer) const
+{
+	return ExternalGameplayTags.HasAny(TagContainer);
+}
+
 FTransform FContextualAnimSceneBindingContext::GetTransform() const
 {
 	// If created with an external transform, used that one to represent the location/rotation of the actor
diff --git a/Engine/Plugins/Experimental/Animation/ContextualAnimation/Source/ContextualAnimation/Public/ContextualAnimTypes.h b/Engine/Plugins/Experimental/Animation/ContextualAnimation/Source/ContextualAnimation/Public/ContextualAnimTypes.h
--- a/Engine/Plugins/Experimental/Animation/ContextualAnimation/Source/ContextualAnimation/Public/ContextualAnimTypes.h
+++ b/Engine/Plugins/Experimental/Animation/ContextualAnimation/Source/ContextualAnimation/Public/ContextualAnimTypes.h
@@ -329,6 +329,9 @@ struct CONTEXTUALANIMATION_API FContextualAnimSceneBindingContext
 	FContextualAnimSceneBindingContext(const FTransform& InExternalTransform, const TOptional<FVector>& InExternalVelocity = TOptional<FVector>())
 		: ExternalTransform(InExternalTransform), ExternalVelocity(InExternalVelocity) {}
 
+	FContextualAnimSceneBindingContext(AActor* InActor, const FGameplayTagContainer& InExternalTags)
+		: Actor(InActor), ExternalGameplayTags(InExternalTags) {}
+
 	AActor* GetActor() const { return Actor.Get(); }
 
 	UAnimInstance* GetAnimInstance() const;
@@ -343,7 +346,15 @@ struct CONTEXTUALANIMATION_API FContextualAnimSceneBindingContext
 
 	FVector GetVelocity() const;
 
-	//@TODO: Add accessors for GameplayTags
+	void AddGameplayTag(const FGameplayTag& Tag);
+	
+	const FGameplayTagContainer& GetGameplayTags() const { return ExternalGameplayTags; }
+
+	bool HasMatchingGameplayTag(const FGameplayTag& TagToCheck) const;
+
+	bool HasAllMatchingGameplayTags(const FGameplayTagContainer& TagContainer) const;
+
+	bool HasAnyMatchingGameplayTags(const FGameplayTagContainer& TagContainer) const;
 
 private:
 
diff --git a/Engine/Plugins/Experimental/Animation/ContextualAnimation/Source/ContextualAnimation/Public/ContextualAnimUtilities.h b/Engine/Plugins/Experimental/Animation/ContextualAnimation/Source/ContextualAnimation/Public/ContextualAnimUtilities.h
--- a/Engine/Plugins/Experimental/Animation/ContextualAnimation/Source/ContextualAnimation/Public/ContextualAnimUtilities.h
+++ b/Engine/Plugins/Experimental/Animation/ContextualAnimation/Source/ContextualAnimation/Public/ContextualAnimUtilities.h
@@ -142,6 +142,19 @@ public:
 	UFUNCTION(BlueprintPure, Category = "Contextual Anim|Scene Binding Context", meta = (DisplayName = "Get Velocity"))
 	static FVector BP_SceneBindingContext_GetVelocity(const FContextualAnimSceneBindingContext& BindingContext) { return BindingContext.GetVelocity(); }
 
+	UFUNCTION(BlueprintPure, Category = "Contextual Anim|Scene Binding Context", meta = (DisplayName = "Get GameplayTags"))
+	static const FGameplayTagContainer& BP_SceneBindingContext_GetGameplayTags(const FContextualAnimSceneBindingContext& BindingContext) { return BindingContext.GetGameplayTags(); }
+
+	UFUNCTION(BlueprintPure, Category = "Contextual Anim|Scene Binding Context", meta = (DisplayName = "Has Matching GameplayTag"))
+	static bool BP_SceneBindingContext_HasMatchingGameplayTag(const FContextualAnimSceneBindingContext& BindingContext, const FGameplayTag& TagToCheck) { return BindingContext.HasMatchingGameplayTag(TagToCheck); }
+
+	UFUNCTION(BlueprintPure, Category = "Contextual Anim|Scene Binding Context", meta = (DisplayName = "Has All Matching GameplayTags"))
+	static bool BP_SceneBindingContext_HasAllMatchingGameplayTags(const FContextualAnimSceneBindingContext& BindingContext, const FGameplayTagContainer& TagContainer) { return BindingContext.HasAllMatchingGameplayTags(TagContainer); }
+
+	UFUNCTION(BlueprintPure, Category = "Contextual Anim|Scene Binding Context", meta = (DisplayName = "Has Any Matching GameplayTags"))
+	static bool BP_SceneBindingContext_HasAnyMatchingGameplayTags(const FContextualAnimSceneBindingContext& BindingContext, const FGameplayTagContainer& TagContainer) { return BindingContext.HasAnyMatchingGameplayTags(TagContainer); }
+
+
 	// FContextualAnimSceneBinding Blueprint Interface
 	//------------------------------------------------------------------------------------------
 
