Contextual Anim: Added check before getting TimeManager from GetWorld as potential fix for a random crash in cooked build.

diff --git a/Engine/Plugins/Experimental/Animation/ContextualAnimation/Source/ContextualAnimation/Private/ContextualAnimSceneActorComponent.cpp b/Engine/Plugins/Experimental/Animation/ContextualAnimation/Source/ContextualAnimation/Private/ContextualAnimSceneActorComponent.cpp
--- a/Engine/Plugins/Experimental/Animation/ContextualAnimation/Source/ContextualAnimation/Private/ContextualAnimSceneActorComponent.cpp
+++ b/Engine/Plugins/Experimental/Animation/ContextualAnimation/Source/ContextualAnimation/Private/ContextualAnimSceneActorComponent.cpp
@@ -87,13 +87,16 @@ void UContextualAnimSceneActorComponent::PlayAnimation_Internal(UAnimSequenceBas
 	// causing this guarding mechanism to fail because by the time the event triggers TGuardValue goes out of the scope
 	// Making our OnMontageBlendingOut to think the animation has been interrupted by an external system and forcing the actor to leave the interaction.
 	bGuardAnimEvents = true;
-	GetWorld()->GetTimerManager().SetTimerForNextTick([WeakThis = MakeWeakObjectPtr(this)]()
+	if (UWorld* World = GetWorld())
 	{
-		if (UContextualAnimSceneActorComponent* Comp = WeakThis.Get())
+		World->GetTimerManager().SetTimerForNextTick([WeakThis = MakeWeakObjectPtr(this)]()
 		{
-			Comp->bGuardAnimEvents = false;
-		}
-	});
+			if (UContextualAnimSceneActorComponent* Comp = WeakThis.Get())
+			{
+				Comp->bGuardAnimEvents = false;
+			}
+		});
+	}
 
 	if (UAnimInstance* AnimInstance = UContextualAnimUtilities::TryGetAnimInstance(GetOwner()))
 	{
