Contextual Anim: Remove old warp target to prevent actor from warping to a wrong location if calculating warp target for current interaction fails + Added logs when warp point calculation fails

diff --git a/Engine/Plugins/Experimental/Animation/ContextualAnimation/Source/ContextualAnimation/Private/ContextualAnimSceneActorComponent.cpp b/Engine/Plugins/Experimental/Animation/ContextualAnimation/Source/ContextualAnimation/Private/ContextualAnimSceneActorComponent.cpp
--- a/Engine/Plugins/Experimental/Animation/ContextualAnimation/Source/ContextualAnimation/Private/ContextualAnimSceneActorComponent.cpp
+++ b/Engine/Plugins/Experimental/Animation/ContextualAnimation/Source/ContextualAnimation/Private/ContextualAnimSceneActorComponent.cpp
@@ -150,6 +150,9 @@ void UContextualAnimSceneActorComponent::AddOrUpdateWarpTargets(int32 SectionIdx
 		return;
 	}
 
+	// Remove old warp targets to prevent actor from warping to a wrong location if calculating warp target for this interaction fails
+	MotionWarpComp->RemoveAllWarpTargets();
+
 	if (const FContextualAnimSceneBinding* Binding = Bindings.FindBindingByActor(GetOwner()))
 	{
 		const UContextualAnimSceneAsset* Asset = Bindings.GetSceneAsset();
diff --git a/Engine/Plugins/Experimental/Animation/ContextualAnimation/Source/ContextualAnimation/Private/ContextualAnimTypes.cpp b/Engine/Plugins/Experimental/Animation/ContextualAnimation/Source/ContextualAnimation/Private/ContextualAnimTypes.cpp
--- a/Engine/Plugins/Experimental/Animation/ContextualAnimation/Source/ContextualAnimation/Private/ContextualAnimTypes.cpp
+++ b/Engine/Plugins/Experimental/Animation/ContextualAnimation/Source/ContextualAnimation/Private/ContextualAnimTypes.cpp
@@ -766,6 +766,11 @@ bool FContextualAnimSceneBindings::CalculateWarpPoint(const FContextualAnimWarpP
 			OutWarpPoint.Transform = PrimaryBinding->GetTransform();
 			return true;
 		}
+		else
+		{
+			UE_LOG(LogContextualAnim, VeryVerbose, TEXT("FContextualAnimSceneBindings::CalculateWarpPoint failed. Reason: Can't find Primary Binding. Asset: %s WarpPointDefinitionMode: PrimaryActor WarpTargetName: %s"),
+				*GetNameSafe(SceneAsset.Get()), *WarpPointDef.WarpTargetName.ToString());
+		}
 	}
 	else if (WarpPointDef.Mode == EContextualAnimWarpPointDefinitionMode::Socket)
 	{
@@ -781,6 +786,16 @@ bool FContextualAnimSceneBindings::CalculateWarpPoint(const FContextualAnimWarpP
 				OutWarpPoint.Transform = WarpPointTransform;
 				return true;
 			}
+			else
+			{
+				UE_LOG(LogContextualAnim, VeryVerbose, TEXT("FContextualAnimSceneBindings::CalculateWarpPoint failed. Reason: Can't find socket used as warp point in primary actor. Asset: %s WarpPointDefinitionMode: Socket WarpTargetName: %s SocketName: %s Primary Actor: %s"),
+					*GetNameSafe(SceneAsset.Get()), *WarpPointDef.WarpTargetName.ToString(), *WarpPointDef.SocketName.ToString(), *GetNameSafe(PrimaryBinding->GetActor()));
+			}
+		}
+		else
+		{
+			UE_LOG(LogContextualAnim, VeryVerbose, TEXT("FContextualAnimSceneBindings::CalculateWarpPoint failed. Reason: Can't find Primary Binding. Asset: %s WarpPointDefinitionMode: Socket WarpTargetName: %s SocketName: %s"),
+				*GetNameSafe(SceneAsset.Get()), *WarpPointDef.WarpTargetName.ToString(), *WarpPointDef.SocketName.ToString());
 		}
 	}
 	else if (WarpPointDef.Mode == EContextualAnimWarpPointDefinitionMode::Custom)
@@ -801,6 +816,11 @@ bool FContextualAnimSceneBindings::CalculateWarpPoint(const FContextualAnimWarpP
 					OutWarpPoint.Transform.SetRotation((T2.GetLocation() - T1.GetLocation()).GetSafeNormal2D().ToOrientationQuat());
 					return true;
 				}
+				else
+				{
+					UE_LOG(LogContextualAnim, VeryVerbose, TEXT("FContextualAnimSceneBindings::CalculateWarpPoint failed. Reason: Can't find binding for Params.OtherRole. Asset: %s WarpPointDefinitionMode: Custom WarpTargetName: %s Params.OtherRole: %s"),
+						*GetNameSafe(SceneAsset.Get()), *WarpPointDef.WarpTargetName.ToString(), *Params.OtherRole.ToString());
+				}
 			}
 			else
 			{
@@ -808,6 +828,11 @@ bool FContextualAnimSceneBindings::CalculateWarpPoint(const FContextualAnimWarpP
 				return true;
 			}
 		}
+		else
+		{
+			UE_LOG(LogContextualAnim, VeryVerbose, TEXT("FContextualAnimSceneBindings::CalculateWarpPoint failed. Reason: Can't find binding for Params.Origin. Asset: %s WarpPointDefinitionMode: Custom WarpTargetName: %s Params.Origin: %s"),
+				*GetNameSafe(SceneAsset.Get()), *WarpPointDef.WarpTargetName.ToString(), *Params.Origin.ToString());
+		}
 	}
 
 	return false;
