Contextual Anim: Marked PlayAnimation function as virtual

diff --git a/Engine/Plugins/Experimental/Animation/ContextualAnimation/Source/ContextualAnimation/Public/ContextualAnimSceneActorComponent.h b/Engine/Plugins/Experimental/Animation/ContextualAnimation/Source/ContextualAnimation/Public/ContextualAnimSceneActorComponent.h
--- a/Engine/Plugins/Experimental/Animation/ContextualAnimation/Source/ContextualAnimation/Public/ContextualAnimSceneActorComponent.h
+++ b/Engine/Plugins/Experimental/Animation/ContextualAnimation/Source/ContextualAnimation/Public/ContextualAnimSceneActorComponent.h
@@ -292,7 +292,7 @@ protected:
 	UFUNCTION(Server, Reliable, WithValidation)
 	void ServerEarlyOutContextualAnimScene();
 
-	void PlayAnimation_Internal(UAnimSequenceBase* Animation, float StartTime, bool bSyncPlaybackTime);
+	virtual void PlayAnimation_Internal(UAnimSequenceBase* Animation, float StartTime, bool bSyncPlaybackTime);
 
 	void AddOrUpdateWarpTargets(int32 SectionIdx, int32 AnimSetIdx, const TArray<FContextualAnimWarpTarget>& ExternalWarpTargets);
 
