Contextual Anim: Temp hack to prevent a crash caused by an actor created by a custom WorldSettings.

diff --git a/Engine/Plugins/Experimental/Animation/ContextualAnimation/Source/ContextualAnimationEditor/Private/ContextualAnimPreviewScene.cpp b/Engine/Plugins/Experimental/Animation/ContextualAnimation/Source/ContextualAnimationEditor/Private/ContextualAnimPreviewScene.cpp
--- a/Engine/Plugins/Experimental/Animation/ContextualAnimation/Source/ContextualAnimationEditor/Private/ContextualAnimPreviewScene.cpp
+++ b/Engine/Plugins/Experimental/Animation/ContextualAnimation/Source/ContextualAnimationEditor/Private/ContextualAnimPreviewScene.cpp
@@ -39,6 +39,18 @@ void FContextualAnimPreviewScene::Tick(float InDeltaTime)
 		}
 
 		PreviewWorld->bBegunPlay = true;
+
+		// @TODO: Temp hacky solution to prevent a crash caused by an actor created by a custom WorldSettings
+		for (TActorIterator<AActor> It(PreviewWorld); It; ++It)
+		{
+			if (AActor* Owner = It->GetOwner())
+			{
+				if (Owner->GetClass()->IsChildOf<AWorldSettings>())
+				{
+					It->Destroy();
+				}
+			}
+		}
 	}
 
 	GetWorld()->Tick(LEVELTICK_All, InDeltaTime);
