Contextual Anim Editor: Removed manual trigger of the begin play event for all the actor in the preview world, which was previously added as workaround for an issue in the movement comp and doesn't seems to be neccesary anymore.

diff --git a/Engine/Plugins/Experimental/Animation/ContextualAnimation/Source/ContextualAnimationEditor/Private/ContextualAnimPreviewScene.cpp b/Engine/Plugins/Experimental/Animation/ContextualAnimation/Source/ContextualAnimationEditor/Private/ContextualAnimPreviewScene.cpp
--- a/Engine/Plugins/Experimental/Animation/ContextualAnimation/Source/ContextualAnimationEditor/Private/ContextualAnimPreviewScene.cpp
+++ b/Engine/Plugins/Experimental/Animation/ContextualAnimation/Source/ContextualAnimationEditor/Private/ContextualAnimPreviewScene.cpp
@@ -29,29 +29,5 @@ void FContextualAnimPreviewScene::Tick(float InDeltaTime)
 {
 	FAdvancedPreviewScene::Tick(InDeltaTime);
 
-	// Trigger Begin Play in this preview world.
-	// This is needed for the CharacterMovementComponent to be able to switch to falling mode. See: UCharacterMovementComponent::StartFalling
-	if (PreviewWorld && !PreviewWorld->bBegunPlay)
-	{
-		for (FActorIterator It(PreviewWorld); It; ++It)
-		{
-			It->DispatchBeginPlay();
-		}
-
-		PreviewWorld->bBegunPlay = true;
-
-		// @TODO: Temp hacky solution to prevent a crash caused by an actor created by a custom WorldSettings
-		for (TActorIterator<AActor> It(PreviewWorld); It; ++It)
-		{
-			if (AActor* Owner = It->GetOwner())
-			{
-				if (Owner->GetClass()->IsChildOf<AWorldSettings>())
-				{
-					It->Destroy();
-				}
-			}
-		}
-	}
-
 	GetWorld()->Tick(LEVELTICK_All, InDeltaTime);
 }
