Contextual Anim: Warp points are now calculated on the server and replicated down to the clients to avoid mismatch when they are relative to moving actor(s)

diff --git a/Engine/Plugins/Experimental/Animation/ContextualAnimation/Source/ContextualAnimation/Private/ContextualAnimSceneActorComponent.cpp b/Engine/Plugins/Experimental/Animation/ContextualAnimation/Source/ContextualAnimation/Private/ContextualAnimSceneActorComponent.cpp
--- a/Engine/Plugins/Experimental/Animation/ContextualAnimation/Source/ContextualAnimation/Private/ContextualAnimSceneActorComponent.cpp
+++ b/Engine/Plugins/Experimental/Animation/ContextualAnimation/Source/ContextualAnimation/Private/ContextualAnimSceneActorComponent.cpp
@@ -41,6 +41,35 @@ void FContextualAnimRepData::IncrementRepCounter()
 	RepCounter = Counter;
 }
 
+static int32 CalculateWarpPointsForBindings(const FContextualAnimSceneBindings& Bindings, int32 SectionIdx, int32 AnimSetIdx, TArray<FContextualAnimWarpPoint>& OutWarpPoints)
+{
+	const UContextualAnimSceneAsset* Asset = Bindings.GetSceneAsset();
+	if (Asset == nullptr)
+	{
+		UE_LOG(LogContextualAnim, Warning, TEXT("CalculateWarpPointsForBindings Invalid Scene Asset. Bindings Id: %d Bindings Num: %d SectionIdx: %d AnimSetIdx: %d"), Bindings.GetID(), Bindings.Num(), SectionIdx, AnimSetIdx);
+		return 0;
+	}
+
+	const FContextualAnimSceneSection* Section = Asset->GetSection(SectionIdx);
+	if (Section == nullptr)
+	{
+		UE_LOG(LogContextualAnim, Warning, TEXT("CalculateWarpPointsForBindings Invalid Section. Bindings Id: %d Bindings Num: %d SectionIdx: %d AnimSetIdx: %d"), Bindings.GetID(), Bindings.Num(), SectionIdx, AnimSetIdx);
+		return 0;
+	}
+
+	OutWarpPoints.Reset(Section->GetWarpPointDefinitions().Num());
+	for (const FContextualAnimWarpPointDefinition& WarpPointDef : Section->GetWarpPointDefinitions())
+	{
+		FContextualAnimWarpPoint WarpPoint;
+		if (Bindings.CalculateWarpPoint(WarpPointDef, WarpPoint))
+		{
+			OutWarpPoints.Add(WarpPoint);
+		}
+	}
+
+	return OutWarpPoints.Num();
+}
+
 UContextualAnimSceneActorComponent::UContextualAnimSceneActorComponent(const FObjectInitializer& ObjectInitializer)
 	: Super(ObjectInitializer)
 {
@@ -140,7 +169,7 @@ void UContextualAnimSceneActorComponent::PlayAnimation_Internal(UAnimSequenceBas
 	}
 }
 
-void UContextualAnimSceneActorComponent::AddOrUpdateWarpTargets(int32 SectionIdx, int32 AnimSetIdx, const TArray<FContextualAnimWarpTarget>& ExternalWarpTargets)
+void UContextualAnimSceneActorComponent::AddOrUpdateWarpTargets(int32 SectionIdx, int32 AnimSetIdx, const TArray<FContextualAnimWarpPoint>& WarpPoints, const TArray<FContextualAnimWarpTarget>& ExternalWarpTargets)
 {
 	// This is relevant only for character with motion warping comp
 	ACharacter* CharacterOwner = Cast<ACharacter>(GetOwner());
@@ -155,37 +184,28 @@ void UContextualAnimSceneActorComponent::AddOrUpdateWarpTargets(int32 SectionIdx
 
 	if (const FContextualAnimSceneBinding* Binding = Bindings.FindBindingByActor(GetOwner()))
 	{
-		const UContextualAnimSceneAsset* Asset = Bindings.GetSceneAsset();
-		if(Asset == nullptr)
-		{
-			UE_LOG(LogContextualAnim, Warning, TEXT("%-21s UContextualAnimSceneActorComponent::AddOrUpdateWarpTargets Invalid Scene Asset. Actor: %s Bindings Id: %d Bindings Num: %d SectionIdx: %d AnimSetIdx: %d"),
-				*UEnum::GetValueAsString(TEXT("Engine.ENetRole"), GetOwner()->GetLocalRole()), *GetNameSafe(GetOwner()), Bindings.GetID(), Bindings.Num(), SectionIdx, AnimSetIdx);
-			return;
-		}
-
-		const FContextualAnimSceneSection* Section = Asset->GetSection(SectionIdx);
-		if (Section == nullptr)
+		if (WarpPoints.Num() > 0)
 		{
-			UE_LOG(LogContextualAnim, Warning, TEXT("%-21s UContextualAnimSceneActorComponent::AddOrUpdateWarpTargets Invalid Section. Actor: %s Bindings Id: %d Bindings Num: %d SectionIdx: %d AnimSetIdx: %d"),
-				*UEnum::GetValueAsString(TEXT("Engine.ENetRole"), GetOwner()->GetLocalRole()), *GetNameSafe(GetOwner()), Bindings.GetID(), Bindings.Num(), SectionIdx, AnimSetIdx);
-			return;
-		}
+			const UContextualAnimSceneAsset* Asset = Bindings.GetSceneAsset();
+			if (Asset == nullptr)
+			{
+				UE_LOG(LogContextualAnim, Warning, TEXT("%-21s UContextualAnimSceneActorComponent::AddOrUpdateWarpTargets Invalid Scene Asset. Actor: %s Bindings Id: %d Bindings Num: %d SectionIdx: %d AnimSetIdx: %d"),
+					*UEnum::GetValueAsString(TEXT("Engine.ENetRole"), GetOwner()->GetLocalRole()), *GetNameSafe(GetOwner()), Bindings.GetID(), Bindings.Num(), SectionIdx, AnimSetIdx);
+				return;
+			}
 
-		if(Section->GetWarpPointDefinitions().Num() > 0)
-		{
 			const FContextualAnimTrack* AnimTrack = Asset->GetAnimTrack(SectionIdx, AnimSetIdx, Bindings.GetRoleFromBinding(*Binding));
 			if (AnimTrack == nullptr || AnimTrack->Animation == nullptr)
 			{
 				return;
 			}
 
-			for (const FContextualAnimWarpPointDefinition& WarpPointDef : Section->GetWarpPointDefinitions())
+			for (const FContextualAnimWarpPoint& WarpPoint : WarpPoints)
 			{
-				FContextualAnimWarpPoint WarpPoint;
-				if (Bindings.CalculateWarpPoint(WarpPointDef, WarpPoint))
+				if (WarpPoint.Name != NAME_None)
 				{
-					const float Time = AnimTrack->GetSyncTimeForWarpSection(WarpPointDef.WarpTargetName);
-					const FTransform TransformRelativeToWarpPoint = Asset->GetAlignmentTransform(*AnimTrack, WarpPointDef.WarpTargetName, Time);
+					const float Time = AnimTrack->GetSyncTimeForWarpSection(WarpPoint.Name);
+					const FTransform TransformRelativeToWarpPoint = Asset->GetAlignmentTransform(*AnimTrack, WarpPoint.Name, Time);
 					const FTransform WarpTargetTransform = TransformRelativeToWarpPoint * WarpPoint.Transform;
 					MotionWarpComp->AddOrUpdateWarpTargetFromTransform(WarpPoint.Name, WarpTargetTransform);
 				}
@@ -205,6 +225,11 @@ void UContextualAnimSceneActorComponent::AddOrUpdateWarpTargets(int32 SectionIdx
 
 bool UContextualAnimSceneActorComponent::LateJoinContextualAnimScene(AActor* Actor, FName Role, const TArray<FContextualAnimWarpTarget>& ExternalWarpTargets)
 {
+	if (!GetOwner()->HasAuthority())
+	{
+		return false;
+	}
+
 	if (!Bindings.IsValid())
 	{
 		UE_LOG(LogContextualAnim, Warning, TEXT("%-21s UContextualAnimSceneActorComponent::LateJoinContextualAnimScene Invalid Bindings"), *UEnum::GetValueAsString(TEXT("Engine.ENetRole"), GetOwner()->GetLocalRole()));
@@ -227,31 +252,10 @@ bool UContextualAnimSceneActorComponent::LateJoinContextualAnimScene(AActor* Act
 	UE_LOG(LogContextualAnim, Verbose, TEXT("%-21s UContextualAnimSceneActorComponent::LateJoinContextualAnimScene Owner: %s Bindings Id: %d Section: %d Asset: %s. Requester: %s Role: %s"),
 		*UEnum::GetValueAsString(TEXT("Engine.ENetRole"), GetOwner()->GetLocalRole()), *GetNameSafe(GetOwner()), Bindings.GetID(), Bindings.GetSectionIdx(), *GetNameSafe(Bindings.GetSceneAsset()), *GetNameSafe(Actor), *Role.ToString());
 
-	// Play animation and set state on this new actor that is joining us and update bindings for everyone else
-	if (HandleLateJoin(Actor, Role, ExternalWarpTargets))
-	{
-		// Replicate late join event. See OnRep_LateJoinData
-		if (GetOwner()->HasAuthority())
-		{
-			RepLateJoinData.Actor = Actor;
-			RepLateJoinData.Role = Role;
-			RepLateJoinData.ExternalWarpTargets = ExternalWarpTargets;
-			RepLateJoinData.IncrementRepCounter();
-			MARK_PROPERTY_DIRTY_FROM_NAME(UContextualAnimSceneActorComponent, RepLateJoinData, this);
-			GetOwner()->ForceNetUpdate();
-		}
-
-		return true;
-	}
-
-	return false;
-}
-
-bool UContextualAnimSceneActorComponent::HandleLateJoin(AActor* Actor, FName Role, const TArray<FContextualAnimWarpTarget>& ExternalWarpTargets)
-{
+	// Add actor to the bindings
 	if (!IsValid(Actor) || !Bindings.BindActorToRole(*Actor, Role))
 	{
-		UE_LOG(LogContextualAnim, Warning, TEXT("%-21s UContextualAnimSceneActorComponent::HandleLateJoin Failed. Reason: Adding %s to the bindings for role: %s failed!"),
+		UE_LOG(LogContextualAnim, Warning, TEXT("%-21s UContextualAnimSceneActorComponent::LateJoinContextualAnimScene Failed. Reason: Adding %s to the bindings for role: %s failed!"),
 			*UEnum::GetValueAsString(TEXT("Engine.ENetRole"), GetOwner()->GetLocalRole()), *GetNameSafe(Actor), *Role.ToString());
 
 		return false;
@@ -269,19 +273,35 @@ bool UContextualAnimSceneActorComponent::HandleLateJoin(AActor* Actor, FName Rol
 		}
 	}
 
+	// For now when late joining an scene always play animation from first section
+	const int32 SectionIdx = 0;
+	const int32 AnimSetIdx = 0;
+
+	TArray<FContextualAnimWarpPoint> WarpPoints;
+	CalculateWarpPointsForBindings(Bindings, SectionIdx, AnimSetIdx, WarpPoints);
+
 	// Play animation and set state on this new actor that is joining us
 	if (const FContextualAnimSceneBinding* Binding = Bindings.FindBindingByActor(Actor))
 	{
 		if (UContextualAnimSceneActorComponent* Comp = Binding->GetSceneActorComponent())
 		{
-			Comp->LateJoinScene(Bindings, ExternalWarpTargets);
+			Comp->LateJoinScene(Bindings, SectionIdx, AnimSetIdx, WarpPoints, ExternalWarpTargets);
 		}
 	}
 
+	// Replicate late join event. See OnRep_LateJoinData
+	RepLateJoinData.Actor = Actor;
+	RepLateJoinData.Role = Role;
+	RepLateJoinData.WarpPoints = MoveTemp(WarpPoints);
+	RepLateJoinData.ExternalWarpTargets = ExternalWarpTargets;
+	RepLateJoinData.IncrementRepCounter();
+	MARK_PROPERTY_DIRTY_FROM_NAME(UContextualAnimSceneActorComponent, RepLateJoinData, this);
+	GetOwner()->ForceNetUpdate();
+
 	return true;
 }
 
-void UContextualAnimSceneActorComponent::LateJoinScene(const FContextualAnimSceneBindings& InBindings, const TArray<FContextualAnimWarpTarget>& ExternalWarpTargets)
+void UContextualAnimSceneActorComponent::LateJoinScene(const FContextualAnimSceneBindings& InBindings, int32 SectionIdx, int32 AnimSetIdx, const TArray<FContextualAnimWarpPoint>& WarpPoints, const TArray<FContextualAnimWarpTarget>& ExternalWarpTargets)
 {
 	if (Bindings.IsValid())
 	{
@@ -298,15 +318,12 @@ void UContextualAnimSceneActorComponent::LateJoinScene(const FContextualAnimScen
 
 		Bindings = InBindings;
 
-		// For now when late joining an scene always play animation from first section
-		const int32 SectionIdx = 0;
-		const int32 AnimSetIdx = 0;
 		const FContextualAnimTrack* AnimTrack = Bindings.GetSceneAsset()->GetAnimTrack(SectionIdx, AnimSetIdx, Bindings.GetRoleFromBinding(*Binding));
 		check(AnimTrack);
 
 		PlayAnimation_Internal(AnimTrack->Animation, 0.f, false);
 
-		AddOrUpdateWarpTargets(SectionIdx, AnimSetIdx, ExternalWarpTargets);
+		AddOrUpdateWarpTargets(SectionIdx, AnimSetIdx, WarpPoints, ExternalWarpTargets);
 
 		SetIgnoreCollisionWithOtherActors(true);
 
@@ -332,8 +349,39 @@ void UContextualAnimSceneActorComponent::OnRep_LateJoinData()
 		return;
 	}
 
-	// Play animation and set state on this new actor that is joining us and update bindings for everyone else
-	HandleLateJoin(RepLateJoinData.Actor, RepLateJoinData.Role, RepLateJoinData.ExternalWarpTargets);
+
+	AActor* Actor = RepLateJoinData.Actor;
+	FName Role = RepLateJoinData.Role;
+
+	// Add actor to the bindings
+	if (!IsValid(Actor) || !Bindings.BindActorToRole(*Actor, Role))
+	{
+		UE_LOG(LogContextualAnim, Warning, TEXT("%-21s UContextualAnimSceneActorComponent::OnRep_LateJoinData Failed. Reason: Adding %s to the bindings for role: %s failed!"),
+			*UEnum::GetValueAsString(TEXT("Engine.ENetRole"), GetOwner()->GetLocalRole()), *GetNameSafe(Actor), *Role.ToString());
+
+		return;
+	}
+
+	// Update the bindings on all the other actors too
+	for (const FContextualAnimSceneBinding& OtherBinding : Bindings)
+	{
+		if (OtherBinding.GetActor() != GetOwner() && OtherBinding.GetActor() != Actor)
+		{
+			if (UContextualAnimSceneActorComponent* Comp = OtherBinding.GetSceneActorComponent())
+			{
+				Comp->Bindings.BindActorToRole(*Actor, Role);
+			}
+		}
+	}
+
+	// Play animation and set state on this new actor that is joining us
+	if (const FContextualAnimSceneBinding* Binding = Bindings.FindBindingByActor(Actor))
+	{
+		if (UContextualAnimSceneActorComponent* Comp = Binding->GetSceneActorComponent())
+		{
+			Comp->LateJoinScene(Bindings, 0, 0, RepLateJoinData.WarpPoints, RepLateJoinData.ExternalWarpTargets);
+		}
+	}
 }
 
 bool UContextualAnimSceneActorComponent::TransitionContextualAnimScene(FName SectionName, int32 AnimSetIdx, const TArray<FContextualAnimWarpTarget>& ExternalWarpTargets)
@@ -364,11 +412,16 @@ bool UContextualAnimSceneActorComponent::TransitionContextualAnimScene(FName Sec
 			UE_LOG(LogContextualAnim, Log, TEXT("%-21s UContextualAnimSceneActorComponent::TransitionTo Actor: %s SectionName: %s"),
 				*UEnum::GetValueAsString(TEXT("Engine.ENetRole"), GetOwner()->GetLocalRole()), *GetNameSafe(GetOwner()), *SectionName.ToString());
 
-			HandleTransitionEveryone(SectionIdx, AnimSetIdx, ExternalWarpTargets);
+			// Calculate WarpPoints
+			TArray<FContextualAnimWarpPoint> WarpPoints;
+			CalculateWarpPointsForBindings(Bindings, SectionIdx, AnimSetIdx, WarpPoints);
+
+			HandleTransitionEveryone(SectionIdx, AnimSetIdx, WarpPoints, ExternalWarpTargets);
 
 			RepTransitionData.Id = Bindings.GetID();
 			RepTransitionData.SectionIdx = SectionIdx;
 			RepTransitionData.AnimSetIdx = AnimSetIdx;
+			RepTransitionData.WarpPoints = MoveTemp(WarpPoints);
 			RepTransitionData.ExternalWarpTargets = ExternalWarpTargets;
 			RepTransitionData.IncrementRepCounter();
 			MARK_PROPERTY_DIRTY_FROM_NAME(UContextualAnimSceneActorComponent, RepTransitionData, this);
@@ -410,11 +463,11 @@ bool UContextualAnimSceneActorComponent::TransitionContextualAnimScene(FName Sec
 	return TransitionContextualAnimScene(SectionName, AnimSetIdx, ExternalWarpTargets);
 }
 
-void UContextualAnimSceneActorComponent::HandleTransitionEveryone(int32 NewSectionIdx, int32 NewAnimSetIdx, const TArray<FContextualAnimWarpTarget>& ExternalWarpTargets)
+void UContextualAnimSceneActorComponent::HandleTransitionEveryone(int32 NewSectionIdx, int32 NewAnimSetIdx, const TArray<FContextualAnimWarpPoint>& WarpPoints, const TArray<FContextualAnimWarpTarget>& ExternalWarpTargets)
 {
 	// Update Bindings internal data and play new animation for the leader first
 	// Note that for now we always transition to the first set in the section. We could run selection criteria here too but keeping it simple for now
-	HandleTransitionSelf(NewSectionIdx, NewAnimSetIdx, ExternalWarpTargets);
+	HandleTransitionSelf(NewSectionIdx, NewAnimSetIdx, WarpPoints, ExternalWarpTargets);
 
 	// And now the same for everyone else
 	for (const FContextualAnimSceneBinding& Binding : Bindings)
@@ -423,13 +476,13 @@ void UContextualAnimSceneActorComponent::HandleTransitionEveryone(int32 NewSecti
 		{
 			if(UContextualAnimSceneActorComponent* Comp = Binding.GetSceneActorComponent())
 			{
-				Comp->HandleTransitionSelf(NewSectionIdx, NewAnimSetIdx, ExternalWarpTargets);
+				Comp->HandleTransitionSelf(NewSectionIdx, NewAnimSetIdx, WarpPoints, ExternalWarpTargets);
 			}
 		}
 	}
 }
 
-void UContextualAnimSceneActorComponent::HandleTransitionSelf(int32 NewSectionIdx, int32 NewAnimSetIdx, const TArray<FContextualAnimWarpTarget>& ExternalWarpTargets)
+void UContextualAnimSceneActorComponent::HandleTransitionSelf(int32 NewSectionIdx, int32 NewAnimSetIdx, const TArray<FContextualAnimWarpPoint>& WarpPoints, const TArray<FContextualAnimWarpTarget>& ExternalWarpTargets)
 {
 	// Update bindings internal data so it points to the new section and new anim set
 	Bindings.TransitionTo(NewSectionIdx, NewAnimSetIdx);
@@ -439,7 +492,7 @@ void UContextualAnimSceneActorComponent::HandleTransitionSelf(int32 NewSectionId
 	const FContextualAnimTrack& AnimTrack = Bindings.GetAnimTrackFromBinding(*Bindings.FindBindingByActor(GetOwner()));
 	PlayAnimation_Internal(AnimTrack.Animation, 0.f, true);
 
-	AddOrUpdateWarpTargets(NewSectionIdx, NewAnimSetIdx, ExternalWarpTargets);
+	AddOrUpdateWarpTargets(NewSectionIdx, NewAnimSetIdx, WarpPoints, ExternalWarpTargets);
 }
 
 bool UContextualAnimSceneActorComponent::TransitionSingleActor(int32 SectionIdx, const TArray<FContextualAnimWarpTarget>& ExternalWarpTargets)
@@ -479,13 +532,18 @@ bool UContextualAnimSceneActorComponent::TransitionSingleActor(int32 SectionIdx,
 				UE_LOG(LogContextualAnim, Log, TEXT("%-21s UContextualAnimSceneActorComponent::TransitionSingleActor Actor: %s SectionIdx: %d AnimSetIdx: %d"),
 					*UEnum::GetValueAsString(TEXT("Engine.ENetRole"), GetOwner()->GetLocalRole()), *GetNameSafe(GetOwner()), SectionIdx, AnimSetIdx);
 
+				// Calculate WarpPoints
+				TArray<FContextualAnimWarpPoint> WarpPoints;
+				CalculateWarpPointsForBindings(Bindings, SectionIdx, AnimSetIdx, WarpPoints);
+
 				PlayAnimation_Internal(AnimTrack->Animation, 0.f, false);
 
-				AddOrUpdateWarpTargets(SectionIdx, AnimSetIdx, ExternalWarpTargets);
+				AddOrUpdateWarpTargets(SectionIdx, AnimSetIdx, WarpPoints, ExternalWarpTargets);
 
 				RepTransitionSingleActorData.Id = Bindings.GetID();
 				RepTransitionSingleActorData.SectionIdx = SectionIdx;
 				RepTransitionSingleActorData.AnimSetIdx = AnimSetIdx;
+				RepTransitionSingleActorData.WarpPoints = MoveTemp(WarpPoints);
 				RepTransitionSingleActorData.ExternalWarpTargets = ExternalWarpTargets;
 				RepTransitionSingleActorData.IncrementRepCounter();
 				MARK_PROPERTY_DIRTY_FROM_NAME(UContextualAnimSceneActorComponent, RepTransitionSingleActorData, this);
@@ -520,7 +578,7 @@ void UContextualAnimSceneActorComponent::OnRep_RepTransitionSingleActor()
 				{
 					PlayAnimation_Internal(AnimTrack->Animation, 0.f, false);
 
-					AddOrUpdateWarpTargets(RepTransitionSingleActorData.SectionIdx, RepTransitionSingleActorData.AnimSetIdx, RepTransitionSingleActorData.ExternalWarpTargets);
+					AddOrUpdateWarpTargets(RepTransitionSingleActorData.SectionIdx, RepTransitionSingleActorData.AnimSetIdx, RepTransitionSingleActorData.WarpPoints, RepTransitionSingleActorData.ExternalWarpTargets);
 				}
 			}
 		}
@@ -568,7 +626,14 @@ bool UContextualAnimSceneActorComponent::StartContextualAnimScene(const FContext
 	{
 		if (GetOwner()->HasAuthority())
 		{
-			JoinScene(InBindings, ExternalWarpTargets);
+			// @TODO: Warp points are calculated on the server and replicated to everyone to avoid mismatch when they are relative to moving actor(s)
+			// It may be better to have them in the Bindings but during late joint the actor that is joining the interaction starts from the first section, 
+			// which could be different from the section the other actors are playing.
+			// We should reconsider all that in the future, maybe moving SectionIdx and AnimSetIdx out of the bindings if we want to support that case long term.
+			TArray<FContextualAnimWarpPoint> WarpPoints;
+			CalculateWarpPointsForBindings(InBindings, InBindings.GetSectionIdx(), InBindings.GetAnimSetIdx(), WarpPoints);
+
+			JoinScene(InBindings, WarpPoints, ExternalWarpTargets);
 
 			for (const FContextualAnimSceneBinding& Binding : InBindings)
 			{
@@ -576,12 +641,13 @@ bool UContextualAnimSceneActorComponent::StartContextualAnimScene(const FContext
 				{
 					if (UContextualAnimSceneActorComponent* Comp = Binding.GetSceneActorComponent())
 					{
-						Comp->JoinScene(InBindings, ExternalWarpTargets);
+						Comp->JoinScene(InBindings, WarpPoints, ExternalWarpTargets);
 					}
 				}
 			}
 
 			RepBindings.Bindings = InBindings;
+			RepBindings.WarpPoints = MoveTemp(WarpPoints);
 			RepBindings.ExternalWarpTargets = ExternalWarpTargets;
 			RepBindings.IncrementRepCounter();
 
@@ -592,10 +658,7 @@ bool UContextualAnimSceneActorComponent::StartContextualAnimScene(const FContext
 		}
 		else if (GetOwner()->GetLocalRole() == ROLE_AutonomousProxy)
 		{
-			JoinScene(InBindings, ExternalWarpTargets);
-
 			ServerStartContextualAnimScene(InBindings);
-
 			return true;
 		}
 	}
@@ -682,7 +745,7 @@ void UContextualAnimSceneActorComponent::OnRep_TransitionData()
 		return;
 	}
 
-	HandleTransitionEveryone(RepTransitionData.SectionIdx, RepTransitionData.AnimSetIdx, RepTransitionData.ExternalWarpTargets);
+	HandleTransitionEveryone(RepTransitionData.SectionIdx, RepTransitionData.AnimSetIdx, RepTransitionData.WarpPoints, RepTransitionData.ExternalWarpTargets);
 }
 
 void UContextualAnimSceneActorComponent::OnRep_Bindings()
@@ -702,7 +765,7 @@ void UContextualAnimSceneActorComponent::OnRep_Bindings()
 		if (ensureAlways(OwnerBinding))
 		{
 			// Join the scene (start playing animation, etc.)
-			JoinScene(RepBindings.Bindings, RepBindings.ExternalWarpTargets);
+			JoinScene(RepBindings.Bindings, RepBindings.WarpPoints, RepBindings.ExternalWarpTargets);
 
 			// RepBindings is only replicated from the initiator of the action.
 			// So now we have to tell everyone else involved in the interaction to join us
@@ -714,7 +777,7 @@ void UContextualAnimSceneActorComponent::OnRep_Bindings()
 				{
 					if (UContextualAnimSceneActorComponent* Comp = Binding.GetSceneActorComponent())
 					{
-						Comp->JoinScene(RepBindings.Bindings, RepBindings.ExternalWarpTargets);
+						Comp->JoinScene(RepBindings.Bindings, RepBindings.WarpPoints, RepBindings.ExternalWarpTargets);
 					}
 				}
 			}
@@ -841,7 +904,7 @@ void UContextualAnimSceneActorComponent::OnLeftScene()
 	}
 }
 
-void UContextualAnimSceneActorComponent::JoinScene(const FContextualAnimSceneBindings& InBindings, const TArray<FContextualAnimWarpTarget>& ExternalWarpTargets)
+void UContextualAnimSceneActorComponent::JoinScene(const FContextualAnimSceneBindings& InBindings, const TArray<FContextualAnimWarpPoint> WarpPoints, const TArray<FContextualAnimWarpTarget>& ExternalWarpTargets)
 {
 	if (Bindings.IsValid())
 	{
@@ -858,7 +921,7 @@ void UContextualAnimSceneActorComponent::JoinScene(const FContextualAnimSceneBin
 		const FContextualAnimTrack& AnimTrack = Bindings.GetAnimTrackFromBinding(*Binding);
 		PlayAnimation_Internal(AnimTrack.Animation, 0.f, true);
 
-		AddOrUpdateWarpTargets(AnimTrack.SectionIdx, AnimTrack.AnimSetIdx, ExternalWarpTargets);
+		AddOrUpdateWarpTargets(AnimTrack.SectionIdx, AnimTrack.AnimSetIdx, WarpPoints, ExternalWarpTargets);
 
 		// Disable collision between actors so they can align perfectly
 		SetIgnoreCollisionWithOtherActors(true);
@@ -978,6 +1041,7 @@ void UContextualAnimSceneActorComponent::OnMontageBlendingOut(UAnimMontage* Mont
 			RepTransitionSingleActorData.Id = BindingsId;
 			RepTransitionSingleActorData.SectionIdx = MAX_uint8;
 			RepTransitionSingleActorData.AnimSetIdx = MAX_uint8;
+			RepTransitionSingleActorData.WarpPoints.Reset();
 			RepTransitionSingleActorData.ExternalWarpTargets.Reset();
 			RepTransitionSingleActorData.IncrementRepCounter();
 
diff --git a/Engine/Plugins/Experimental/Animation/ContextualAnimation/Source/ContextualAnimation/Public/ContextualAnimSceneActorComponent.h b/Engine/Plugins/Experimental/Animation/ContextualAnimation/Source/ContextualAnimation/Public/ContextualAnimSceneActorComponent.h
--- a/Engine/Plugins/Experimental/Animation/ContextualAnimation/Source/ContextualAnimation/Public/ContextualAnimSceneActorComponent.h
+++ b/Engine/Plugins/Experimental/Animation/ContextualAnimation/Source/ContextualAnimation/Public/ContextualAnimSceneActorComponent.h
@@ -66,6 +66,9 @@ struct FContextualAnimRepBindingsData : public FContextualAnimRepData
 	UPROPERTY()
 	FContextualAnimSceneBindings Bindings;
 
+	UPROPERTY()
+	TArray<FContextualAnimWarpPoint> WarpPoints;
+
 	UPROPERTY()
 	TArray<FContextualAnimWarpTarget> ExternalWarpTargets;
 
@@ -73,6 +76,7 @@ struct FContextualAnimRepBindingsData : public FContextualAnimRepData
 	{
 		RepCounter = 0;
 		Bindings.Reset();
+		WarpPoints.Reset();
 		ExternalWarpTargets.Reset();
 	}
 };
@@ -91,6 +95,9 @@ struct FContextualAnimRepLateJoinData : public FContextualAnimRepData
 	UPROPERTY()
 	FName Role = NAME_None;
 
+	UPROPERTY()
+	TArray<FContextualAnimWarpPoint> WarpPoints;
+
 	UPROPERTY()
 	TArray<FContextualAnimWarpTarget> ExternalWarpTargets;
 
@@ -99,6 +106,7 @@ struct FContextualAnimRepLateJoinData : public FContextualAnimRepData
 		RepCounter = 0;
 		Actor = nullptr;
 		Role = NAME_None;
+		WarpPoints.Reset();
 		ExternalWarpTargets.Reset();
 	}
 };
@@ -118,6 +126,9 @@ struct FContextualAnimRepTransitionData : public FContextualAnimRepData
 	UPROPERTY()
 	uint8 AnimSetIdx = 0;
 
+	UPROPERTY()
+	TArray<FContextualAnimWarpPoint> WarpPoints;
+
 	UPROPERTY()
 	TArray<FContextualAnimWarpTarget> ExternalWarpTargets;
 
@@ -127,6 +138,7 @@ struct FContextualAnimRepTransitionData : public FContextualAnimRepData
 		Id = 0;
 		SectionIdx = 0;
 		AnimSetIdx = 0;
+		WarpPoints.Reset();
 		ExternalWarpTargets.Reset();
 	}
 };
@@ -278,13 +290,11 @@ protected:
 	// @TODO: These two functions are going to replace OnJoinedScene and OnLeftScene
 	// main different is that these new functions are taking care of animation playback too
 
-	void JoinScene(const FContextualAnimSceneBindings& InBindings, const TArray<FContextualAnimWarpTarget>& ExternalWarpTargets);
+	void JoinScene(const FContextualAnimSceneBindings& InBindings, const TArray<FContextualAnimWarpPoint> WarpPoints, const TArray<FContextualAnimWarpTarget>& ExternalWarpTargets);
 
 	void LeaveScene();
 
-	void LateJoinScene(const FContextualAnimSceneBindings& InBindings, const TArray<FContextualAnimWarpTarget>& ExternalWarpTargets);
-
-	bool HandleLateJoin(AActor* Actor, FName Role, const TArray<FContextualAnimWarpTarget>& ExternalWarpTargets);
+	void LateJoinScene(const FContextualAnimSceneBindings& InBindings, int32 SectionIdx, int32 AnimSetIdx, const TArray<FContextualAnimWarpPoint>& WarpPoints, const TArray<FContextualAnimWarpTarget>& ExternalWarpTargets);
 
 	UFUNCTION(Server, Reliable, WithValidation)
 	void ServerStartContextualAnimScene(const FContextualAnimSceneBindings& InBindings);
@@ -294,11 +304,11 @@ protected:
 
 	virtual void PlayAnimation_Internal(UAnimSequenceBase* Animation, float StartTime, bool bSyncPlaybackTime);
 
-	void AddOrUpdateWarpTargets(int32 SectionIdx, int32 AnimSetIdx, const TArray<FContextualAnimWarpTarget>& ExternalWarpTargets);
+	void AddOrUpdateWarpTargets(int32 SectionIdx, int32 AnimSetIdx, const TArray<FContextualAnimWarpPoint>& WarpPoints, const TArray<FContextualAnimWarpTarget>& ExternalWarpTargets);
 
-	void HandleTransitionSelf(int32 NewSectionIdx, int32 NewAnimSetIdx, const TArray<FContextualAnimWarpTarget>& ExternalWarpTargets);
+	void HandleTransitionSelf(int32 NewSectionIdx, int32 NewAnimSetIdx, const TArray<FContextualAnimWarpPoint>& WarpPoints, const TArray<FContextualAnimWarpTarget>& ExternalWarpTargets);
 
-	void HandleTransitionEveryone(int32 NewSectionIdx, int32 NewAnimSetIdx, const TArray<FContextualAnimWarpTarget>& ExternalWarpTargets);
+	void HandleTransitionEveryone(int32 NewSectionIdx, int32 NewAnimSetIdx, const TArray<FContextualAnimWarpPoint>& WarpPoints, const TArray<FContextualAnimWarpTarget>& ExternalWarpTargets);
 
 private:
 
