Contextual Anim: Replaced simple IgnoreCollisionBetweenActor option with multiple options for dealing with collisions, adding a new option to allow users to define specific collision channels that should be ignored by any of the actors in the interaction

diff --git a/Engine/Plugins/Experimental/Animation/ContextualAnimation/Source/ContextualAnimation/Private/ContextualAnimSceneActorComponent.cpp b/Engine/Plugins/Experimental/Animation/ContextualAnimation/Source/ContextualAnimation/Private/ContextualAnimSceneActorComponent.cpp
--- a/Engine/Plugins/Experimental/Animation/ContextualAnimation/Source/ContextualAnimation/Private/ContextualAnimSceneActorComponent.cpp
+++ b/Engine/Plugins/Experimental/Animation/ContextualAnimation/Source/ContextualAnimation/Private/ContextualAnimSceneActorComponent.cpp
@@ -325,7 +325,7 @@ void UContextualAnimSceneActorComponent::LateJoinScene(const FContextualAnimScen
 
 		AddOrUpdateWarpTargets(SectionIdx, AnimSetIdx, WarpPoints, ExternalWarpTargets);
 
-		SetIgnoreCollisionWithOtherActors(true);
+		SetCollisionState();
 
 		SetMovementState(AnimTrack->bRequireFlyingMode);
 	}
@@ -833,75 +833,71 @@ void UContextualAnimSceneActorComponent::SetIgnoreCollisionWithOtherActors(bool
 	}
 }
 
-void UContextualAnimSceneActorComponent::OnJoinedScene(const FContextualAnimSceneBindings& InBindings)
+void UContextualAnimSceneActorComponent::SetCollisionState()
 {
-	UE_LOG(LogContextualAnim, Verbose, TEXT("%-21s UContextualAnimSceneActorComponent::OnJoinedScene Actor: %s InBindings Id: %d"),
-		*UEnum::GetValueAsString(TEXT("Engine.ENetRole"), GetOwner()->GetLocalRole()), *GetNameSafe(GetOwner()), InBindings.GetID());
-
-	if (Bindings.IsValid())
-	{
-		OnLeftScene();
-	}
-
-	if (const FContextualAnimSceneBinding* Binding = InBindings.FindBindingByActor(GetOwner()))
+	if (const UContextualAnimSceneAsset* Asset = Bindings.GetSceneAsset())
 	{
-		Bindings = InBindings;
-
-		USkeletalMeshComponent* SkelMeshComp = UContextualAnimUtilities::TryGetSkeletalMeshComponent(GetOwner());
-		if (SkelMeshComp && !SkelMeshComp->OnTickPose.IsBoundToObject(this))
+		const EContextualAnimCollisionBehavior CollisionBehavior = Asset->GetCollisionBehavior();
+		if (CollisionBehavior == EContextualAnimCollisionBehavior::IgnoreActorWhenMoving)
 		{
-			SkelMeshComp->OnTickPose.AddUObject(this, &UContextualAnimSceneActorComponent::OnTickPose);
+			SetIgnoreCollisionWithOtherActors(true);
 		}
-
-		// Disable collision between actors so they can align perfectly
-		SetIgnoreCollisionWithOtherActors(true);
-
-		// Prevent physics rotation. During the interaction we want to be fully root motion driven
-		if (UCharacterMovementComponent* MovementComp = GetOwner()->FindComponentByClass<UCharacterMovementComponent>())
+		else if (CollisionBehavior == EContextualAnimCollisionBehavior::IgnoreChannels)
 		{
-			CharacterPropertiesBackup.bAllowPhysicsRotationDuringAnimRootMotion = MovementComp->bAllowPhysicsRotationDuringAnimRootMotion;
-			CharacterPropertiesBackup.bUseControllerDesiredRotation = MovementComp->bUseControllerDesiredRotation;
-			CharacterPropertiesBackup.bOrientRotationToMovement = MovementComp->bOrientRotationToMovement;
-			MovementComp->bAllowPhysicsRotationDuringAnimRootMotion = false;
-			MovementComp->bUseControllerDesiredRotation = false;
-			MovementComp->bOrientRotationToMovement = false;
-		}
+			if (UPrimitiveComponent* RootPrimitiveComponent = Cast<UPrimitiveComponent>(GetOwner()->GetRootComponent()))
+			{
+				if (const FContextualAnimSceneBinding* Binding = Bindings.FindBindingByActor(GetOwner()))
+				{
+					const TArray<TEnumAsByte<ECollisionChannel>>& ChannelsToIgnore = Asset->GetCollisionChannelsToIgnoreForRole(Bindings.GetRoleFromBinding(*Binding));
+					if (ChannelsToIgnore.Num() > 0)
+					{
+						CharacterPropertiesBackup.CollisionResponses.Reset(ChannelsToIgnore.Num());
+						for (ECollisionChannel Channel : ChannelsToIgnore)
+						{
+							ECollisionResponse Response = RootPrimitiveComponent->GetCollisionResponseToChannel(Channel);
+							CharacterPropertiesBackup.CollisionResponses.Add(MakeTuple(Channel, Response));
 
-		OnJoinedSceneDelegate.Broadcast(this);
+							RootPrimitiveComponent->SetCollisionResponseToChannel(Channel, ECR_Ignore);
+						}
+					}
+				}
+			}
+		}
 	}
 }
 
-void UContextualAnimSceneActorComponent::OnLeftScene()
+void UContextualAnimSceneActorComponent::RestoreCollisionState()
 {
-	UE_LOG(LogContextualAnim, Verbose, TEXT("%-21s UContextualAnimSceneActorComponent::OnLeftScene Actor: %s Current Bindings Id: %d"),
-		*UEnum::GetValueAsString(TEXT("Engine.ENetRole"), GetOwner()->GetLocalRole()), *GetNameSafe(GetOwner()), Bindings.GetID());
-
-	if (const FContextualAnimSceneBinding* Binding = Bindings.FindBindingByActor(GetOwner()))
+	if (const UContextualAnimSceneAsset* Asset = Bindings.GetSceneAsset())
 	{
-		// Stop listening to TickPose if we were
-		USkeletalMeshComponent* SkelMeshComp = UContextualAnimUtilities::TryGetSkeletalMeshComponent(GetOwner());
-		if (SkelMeshComp && SkelMeshComp->OnTickPose.IsBoundToObject(this))
+		const EContextualAnimCollisionBehavior CollisionBehavior = Asset->GetCollisionBehavior();
+		if (CollisionBehavior == EContextualAnimCollisionBehavior::IgnoreActorWhenMoving)
 		{
-			SkelMeshComp->OnTickPose.RemoveAll(this);
+			SetIgnoreCollisionWithOtherActors(false);
 		}
-
-		// Restore collision between actors
-		// Note that this assumes that we are the only one disabling the collision between these actors. 
-		// We might want to add a more robust mechanism to avoid overriding a request to disable collision that may have been set by another system
-		SetIgnoreCollisionWithOtherActors(false);
-
-		// Restore bAllowPhysicsRotationDuringAnimRootMotion
-		if (UCharacterMovementComponent* MovementComp = GetOwner()->FindComponentByClass<UCharacterMovementComponent>())
+		else if (CollisionBehavior == EContextualAnimCollisionBehavior::IgnoreChannels)
 		{
-			MovementComp->bAllowPhysicsRotationDuringAnimRootMotion = CharacterPropertiesBackup.bAllowPhysicsRotationDuringAnimRootMotion;
-			MovementComp->bUseControllerDesiredRotation = CharacterPropertiesBackup.bUseControllerDesiredRotation;
-			MovementComp->bOrientRotationToMovement = CharacterPropertiesBackup.bOrientRotationToMovement;
+			if (UPrimitiveComponent* RootPrimitiveComponent = Cast<UPrimitiveComponent>(GetOwner()->GetRootComponent()))
+			{
+				for (const TTuple<ECollisionChannel, ECollisionResponse>& Response : CharacterPropertiesBackup.CollisionResponses)
+				{
+					RootPrimitiveComponent->SetCollisionResponseToChannel(Response.Get<0>(), Response.Get<1>());
+				}
+
+				CharacterPropertiesBackup.CollisionResponses.Reset();
+			}
 		}
+	}
+}
 
-		OnLeftSceneDelegate.Broadcast(this);
+void UContextualAnimSceneActorComponent::OnJoinedScene(const FContextualAnimSceneBindings& InBindings)
+{
+	// This function will be removed
+}
 
-		Bindings.Reset();
-	}
+void UContextualAnimSceneActorComponent::OnLeftScene()
+{
+	// This function will be removed
 }
 
 void UContextualAnimSceneActorComponent::JoinScene(const FContextualAnimSceneBindings& InBindings, const TArray<FContextualAnimWarpPoint> WarpPoints, const TArray<FContextualAnimWarpTarget>& ExternalWarpTargets)
@@ -923,8 +919,7 @@ void UContextualAnimSceneActorComponent::JoinScene(const FContextualAnimSceneBin
 
 		AddOrUpdateWarpTargets(AnimTrack.SectionIdx, AnimTrack.AnimSetIdx, WarpPoints, ExternalWarpTargets);
 
-		// Disable collision between actors so they can align perfectly
-		SetIgnoreCollisionWithOtherActors(true);
+		SetCollisionState();
 
 		SetMovementState(AnimTrack.bRequireFlyingMode);
 
@@ -961,10 +956,7 @@ void UContextualAnimSceneActorComponent::LeaveScene()
 			SkelMeshComp->OnTickPose.RemoveAll(this);
 		}
 		
-		// Restore collision between actors
-		// Note that this assumes that we are the only one disabling the collision between these actors. 
-		// We might want to add a more robust mechanism to avoid overriding a request to disable collision that may have been set by another system
-		SetIgnoreCollisionWithOtherActors(false);
+		RestoreCollisionState();
 
 		const FContextualAnimTrack& AnimTrack = Bindings.GetAnimTrackFromBinding(*Binding);
 		RestoreMovementState(AnimTrack.bRequireFlyingMode);
diff --git a/Engine/Plugins/Experimental/Animation/ContextualAnimation/Source/ContextualAnimation/Private/ContextualAnimSceneAsset.cpp b/Engine/Plugins/Experimental/Animation/ContextualAnimation/Source/ContextualAnimation/Private/ContextualAnimSceneAsset.cpp
--- a/Engine/Plugins/Experimental/Animation/ContextualAnimation/Source/ContextualAnimation/Private/ContextualAnimSceneAsset.cpp
+++ b/Engine/Plugins/Experimental/Animation/ContextualAnimation/Source/ContextualAnimation/Private/ContextualAnimSceneAsset.cpp
@@ -421,7 +421,6 @@ const FContextualAnimTrack* FContextualAnimSceneSection::FindFirstAnimTrackForRo
 UContextualAnimSceneAsset::UContextualAnimSceneAsset(const FObjectInitializer& ObjectInitializer)
 	: Super(ObjectInitializer)
 {
-	bDisableCollisionBetweenActors = true;
 	SampleRate = 15;
 }
 
@@ -892,6 +891,20 @@ void UContextualAnimSceneAsset::GetAlignmentPointsForSecondaryRoleConsideringSel
 	}
 }
 
+const TArray<TEnumAsByte<ECollisionChannel>>& UContextualAnimSceneAsset::GetCollisionChannelsToIgnoreForRole(FName Role) const
+{
+	if (CollisionBehavior == EContextualAnimCollisionBehavior::IgnoreChannels)
+	{
+		if (const FContextualAnimIgnoreChannelsParam* Param = CollisionChannelsToIgnoreParams.FindByPredicate([Role](const FContextualAnimIgnoreChannelsParam& Item){ return Item.Role == Role; }))
+		{
+			return Param->Channels;
+		}
+	}
+
+	static const TArray<TEnumAsByte<ECollisionChannel>> EmptyArray;
+	return EmptyArray;
+}
+
 // Blueprint Interface
 //------------------------------------------------------------------------------------------
 
diff --git a/Engine/Plugins/Experimental/Animation/ContextualAnimation/Source/ContextualAnimation/Private/ContextualAnimSceneInstance.cpp b/Engine/Plugins/Experimental/Animation/ContextualAnimation/Source/ContextualAnimation/Private/ContextualAnimSceneInstance.cpp
--- a/Engine/Plugins/Experimental/Animation/ContextualAnimation/Source/ContextualAnimation/Private/ContextualAnimSceneInstance.cpp
+++ b/Engine/Plugins/Experimental/Animation/ContextualAnimation/Source/ContextualAnimation/Private/ContextualAnimSceneInstance.cpp
@@ -133,7 +133,7 @@ float UContextualAnimSceneInstance::Join(FContextualAnimSceneBinding& Binding)
 		}
 	}
 
-	if (SceneAsset->GetDisableCollisionBetweenActors())
+	if (SceneAsset->GetCollisionBehavior() == EContextualAnimCollisionBehavior::IgnoreActorWhenMoving)
 	{
 		SetIgnoreCollisionWithOtherActors(Actor, true);
 	}
@@ -299,7 +299,7 @@ void UContextualAnimSceneInstance::OnMontageBlendingOut(UAnimMontage* Montage, b
 					}
 				}
 
-				if (SceneAsset->GetDisableCollisionBetweenActors())
+				if (SceneAsset->GetCollisionBehavior() == EContextualAnimCollisionBehavior::IgnoreActorWhenMoving)
 				{
 					SetIgnoreCollisionWithOtherActors(Actor, false);
 				}
diff --git a/Engine/Plugins/Experimental/Animation/ContextualAnimation/Source/ContextualAnimation/Public/ContextualAnimSceneActorComponent.h b/Engine/Plugins/Experimental/Animation/ContextualAnimation/Source/ContextualAnimation/Public/ContextualAnimSceneActorComponent.h
--- a/Engine/Plugins/Experimental/Animation/ContextualAnimation/Source/ContextualAnimation/Public/ContextualAnimSceneActorComponent.h
+++ b/Engine/Plugins/Experimental/Animation/ContextualAnimation/Source/ContextualAnimation/Public/ContextualAnimSceneActorComponent.h
@@ -252,6 +252,7 @@ protected:
 		bool bAllowPhysicsRotationDuringAnimRootMotion = false;
 		bool bUseControllerDesiredRotation = false;
 		bool bOrientRotationToMovement = false;
+		TArray<TTuple<ECollisionChannel, ECollisionResponse>> CollisionResponses;
 	};
 	FCharacterProperties CharacterPropertiesBackup;
 
@@ -277,6 +278,8 @@ protected:
 	void OnRep_TransitionData();
 
 	void SetIgnoreCollisionWithOtherActors(bool bValue) const;
+	void SetCollisionState();
+	void RestoreCollisionState();
 
 	void SetMovementState(bool bRequireFlyingMode);
 	void RestoreMovementState(bool bRequireFlyingMode);
diff --git a/Engine/Plugins/Experimental/Animation/ContextualAnimation/Source/ContextualAnimation/Public/ContextualAnimSceneAsset.h b/Engine/Plugins/Experimental/Animation/ContextualAnimation/Source/ContextualAnimation/Public/ContextualAnimSceneAsset.h
--- a/Engine/Plugins/Experimental/Animation/ContextualAnimation/Source/ContextualAnimation/Public/ContextualAnimSceneAsset.h
+++ b/Engine/Plugins/Experimental/Animation/ContextualAnimation/Source/ContextualAnimation/Public/ContextualAnimSceneAsset.h
@@ -9,6 +9,26 @@
 class UContextualAnimSceneInstance;
 class UContextualAnimSceneAsset;
 
+UENUM(BlueprintType)
+enum class EContextualAnimCollisionBehavior : uint8
+{
+	None,
+	IgnoreActorWhenMoving,
+	IgnoreChannels
+};
+
+USTRUCT(BlueprintType)
+struct FContextualAnimIgnoreChannelsParam
+{
+	GENERATED_BODY()
+
+	UPROPERTY(EditAnywhere, Category = "Defaults")
+	FName Role = NAME_None;
+
+	UPROPERTY(EditAnywhere, Category = "Defaults")
+	TArray<TEnumAsByte<ECollisionChannel>> Channels;
+};
+
 UCLASS(Blueprintable)
 class CONTEXTUALANIMATION_API UContextualAnimRolesAsset : public UDataAsset
 {
@@ -195,13 +215,14 @@ public:
 	void ForEachAnimTrack(FForEachAnimTrackFunction Function) const;
 
 	FORCEINLINE const FName& GetPrimaryRole() const { return PrimaryRole; }
-	
-	FORCEINLINE bool GetDisableCollisionBetweenActors() const { return bDisableCollisionBetweenActors; }
+	FORCEINLINE EContextualAnimCollisionBehavior GetCollisionBehavior() const { return CollisionBehavior; }
 	FORCEINLINE const TSubclassOf<UContextualAnimSceneInstance>& GetSceneInstanceClass() const { return SceneInstanceClass; }
 	FORCEINLINE int32 GetSampleRate() const { return SampleRate; }
 	FORCEINLINE float GetRadius() const { return Radius; }
 	FORCEINLINE bool ShouldPrecomputeAlignmentTracks() const { return bPrecomputeAlignmentTracks; }
 
+	const TArray<TEnumAsByte<ECollisionChannel>>& GetCollisionChannelsToIgnoreForRole(FName Role) const;
+
 	bool HasValidData() const { return RolesAsset != nullptr && Sections.Num() > 0 && Sections[0].AnimSets.Num() > 0; }
 
 	const UContextualAnimRolesAsset* GetRolesAsset() const { return RolesAsset; }
@@ -299,7 +320,10 @@ protected:
 	TSubclassOf<UContextualAnimSceneInstance> SceneInstanceClass;
 
 	UPROPERTY(EditAnywhere, Category = "Settings")
-	bool bDisableCollisionBetweenActors = true;
+	EContextualAnimCollisionBehavior CollisionBehavior = EContextualAnimCollisionBehavior::IgnoreActorWhenMoving;
+
+	UPROPERTY(EditAnywhere, Category = "Settings", meta = (EditCondition = "CollisionBehavior==EContextualAnimCollisionBehavior::IgnoreChannels", EditConditionHides))
+	TArray<FContextualAnimIgnoreChannelsParam> CollisionChannelsToIgnoreParams;
 
 	/** Whether we should extract and cache alignment tracks off line. */
 	UPROPERTY(EditAnywhere, Category = "Settings", AdvancedDisplay)
