Remove unwanted ReloadTextureResources that clear cached for all textures and atlas.

diff --git a/Engine/Plugins/Bridge/Source/Bridge/Private/UI/BridgeStyle.cpp b/Engine/Plugins/Bridge/Source/Bridge/Private/UI/BridgeStyle.cpp
--- a/Engine/Plugins/Bridge/Source/Bridge/Private/UI/BridgeStyle.cpp
+++ b/Engine/Plugins/Bridge/Source/Bridge/Private/UI/BridgeStyle.cpp
@@ -63,8 +63,6 @@ void FBridgeStyle::SetIcon(const FString& StyleName, const FString& ResourcePath
 
 	Name += ".Small";
 	Style->Set(*Name, new IMAGE_BRUSH(ResourcePath, Icon20x20));
-
-	FSlateApplication::Get().GetRenderer()->ReloadTextureResources();
 }
 
 void FBridgeStyle::SetSVGIcon(const FString& StyleName, const FString& ResourcePath)
@@ -77,8 +75,6 @@ void FBridgeStyle::SetSVGIcon(const FString& StyleName, const FString& ResourceP
 
 	Name += ".Small";
 	Style->Set(*Name, new IMAGE_BRUSH_SVG(ResourcePath, Icon20x20));
-
-	FSlateApplication::Get().GetRenderer()->ReloadTextureResources();
 }
 
 #undef IMAGE_BRUSH
diff --git a/Engine/Plugins/Experimental/ActorPalette/Source/ActorPalette/Private/ActorPaletteModule.cpp b/Engine/Plugins/Experimental/ActorPalette/Source/ActorPalette/Private/ActorPaletteModule.cpp
--- a/Engine/Plugins/Experimental/ActorPalette/Source/ActorPalette/Private/ActorPaletteModule.cpp
+++ b/Engine/Plugins/Experimental/ActorPalette/Source/ActorPalette/Private/ActorPaletteModule.cpp
@@ -49,7 +49,6 @@
 void FActorPaletteModule::StartupModule()
 {
 	FActorPaletteStyle::Initialize();
-	FActorPaletteStyle::ReloadTextures();
 
 	FActorPaletteCommands::Register();
 	
diff --git a/Engine/Plugins/Experimental/Animation/ContextualAnimation/Source/ContextualAnimationEditor/Private/ContextualAnimationEditor.cpp b/Engine/Plugins/Experimental/Animation/ContextualAnimation/Source/ContextualAnimationEditor/Private/ContextualAnimationEditor.cpp
--- a/Engine/Plugins/Experimental/Animation/ContextualAnimation/Source/ContextualAnimationEditor/Private/ContextualAnimationEditor.cpp
+++ b/Engine/Plugins/Experimental/Animation/ContextualAnimation/Source/ContextualAnimationEditor/Private/ContextualAnimationEditor.cpp
@@ -21,8 +21,6 @@ void FContextualAnimationEditorModule::StartupModule()
 
 	FContextualAnimEditorStyle::Initialize();
 
-	FContextualAnimEditorStyle::ReloadTextures();
-
 	// Register Ed Mode used to interact with the preview scene
 	FEditorModeRegistry::Get().RegisterMode<FContextualAnimEdMode>(FContextualAnimEdMode::EdModeId, 
 		LOCTEXT("ContextualAnimEdModeEdModeName", "ContextualAnim"), 
diff --git a/Engine/Plugins/Experimental/ChaosVD/Source/ChaosVD/Private/ChaosVDModule.cpp b/Engine/Plugins/Experimental/ChaosVD/Source/ChaosVD/Private/ChaosVDModule.cpp
--- a/Engine/Plugins/Experimental/ChaosVD/Source/ChaosVD/Private/ChaosVDModule.cpp
+++ b/Engine/Plugins/Experimental/ChaosVD/Source/ChaosVD/Private/ChaosVDModule.cpp
@@ -27,7 +27,6 @@ FChaosVDModule& FChaosVDModule::Get()
 void FChaosVDModule::StartupModule()
 {	
 	FChaosVDStyle::Initialize();
-	FChaosVDStyle::ReloadTextures();
 
 	RegisterClassesCustomDetails();
 
