Contextual Anim: Replaced bRequireFlyingMode with MovementMode enum to allow user to set any movement mode during the interaction + Started adding support for attachment

#virtualized

diff --git a/Engine/Plugins/Experimental/Animation/ContextualAnimation/Source/ContextualAnimation/Private/ContextualAnimSceneActorComponent.cpp b/Engine/Plugins/Experimental/Animation/ContextualAnimation/Source/ContextualAnimation/Private/ContextualAnimSceneActorComponent.cpp
--- a/Engine/Plugins/Experimental/Animation/ContextualAnimation/Source/ContextualAnimation/Private/ContextualAnimSceneActorComponent.cpp
+++ b/Engine/Plugins/Experimental/Animation/ContextualAnimation/Source/ContextualAnimation/Private/ContextualAnimSceneActorComponent.cpp
@@ -326,12 +326,19 @@ void UContextualAnimSceneActorComponent::LateJoinScene(const FContextualAnimScen
 
 		AddOrUpdateWarpTargets(SectionIdx, AnimSetIdx, WarpPoints, ExternalWarpTargets);
 
-		SetCollisionState();
+		SetCollisionState(*Binding);
 
-		SetMovementState(AnimTrack->bRequireFlyingMode);
+		SetMovementState(*Binding, AnimTrack->MovementMode);
+
+		OnLateJoinScene(*Binding, SectionIdx, AnimSetIdx);
 	}
 }
 
+void UContextualAnimSceneActorComponent::OnLateJoinScene(const FContextualAnimSceneBinding& Binding, int32 SectionIdx, int32 AnimSetIdx)
+{
+	// For derived classes to override.
+}
+
 void UContextualAnimSceneActorComponent::OnRep_LateJoinData()
 {
 	// This is received by the leader of the interaction on every remote client
@@ -489,11 +496,26 @@ void UContextualAnimSceneActorComponent::HandleTransitionSelf(int32 NewSectionId
 	Bindings.TransitionTo(NewSectionIdx, NewAnimSetIdx);
 
 	// Play animation
-	//@TODO: Add support for dynamic montage
-	const FContextualAnimTrack& AnimTrack = Bindings.GetAnimTrackFromBinding(*Bindings.FindBindingByActor(GetOwner()));
+	const FContextualAnimSceneBinding& Binding = *Bindings.FindBindingByActor(GetOwner());
+	const FContextualAnimTrack& AnimTrack = Bindings.GetAnimTrackFromBinding(Binding);
 	PlayAnimation_Internal(AnimTrack.Animation, 0.f, true);
 
 	AddOrUpdateWarpTargets(NewSectionIdx, NewAnimSetIdx, WarpPoints, ExternalWarpTargets);
+
+	if (UCharacterMovementComponent* MovementComp = Binding.GetCharacterMovementComponent())
+	{
+		if (MovementComp->MovementMode != AnimTrack.MovementMode)
+		{
+			MovementComp->SetMovementMode(AnimTrack.MovementMode);
+		}
+	}
+
+	OnTransitionScene(Binding, NewSectionIdx, NewAnimSetIdx);
+}
+
+void UContextualAnimSceneActorComponent::OnTransitionScene(const FContextualAnimSceneBinding& Binding, int32 NewSectionIdx, int32 NewAnimSetIdx)
+{
+	// For derived classes to override.
 }
 
 bool UContextualAnimSceneActorComponent::TransitionSingleActor(int32 SectionIdx, const TArray<FContextualAnimWarpTarget>& ExternalWarpTargets)
@@ -541,6 +563,16 @@ bool UContextualAnimSceneActorComponent::TransitionSingleActor(int32 SectionIdx,
 
 				AddOrUpdateWarpTargets(SectionIdx, AnimSetIdx, WarpPoints, ExternalWarpTargets);
 
+				if (UCharacterMovementComponent* MovementComp = OwnerBinding->GetCharacterMovementComponent())
+				{
+					if (MovementComp->MovementMode != AnimTrack->MovementMode)
+					{
+						MovementComp->SetMovementMode(AnimTrack->MovementMode);
+					}
+				}
+
+				OnTransitionSingleActor(*OwnerBinding, SectionIdx, AnimSetIdx);
+
 				RepTransitionSingleActorData.Id = Bindings.GetID();
 				RepTransitionSingleActorData.SectionIdx = SectionIdx;
 				RepTransitionSingleActorData.AnimSetIdx = AnimSetIdx;
@@ -558,6 +590,11 @@ bool UContextualAnimSceneActorComponent::TransitionSingleActor(int32 SectionIdx,
 	return false;
 }
 
+void UContextualAnimSceneActorComponent::OnTransitionSingleActor(const FContextualAnimSceneBinding& Binding, int32 NewSectionIdx, int32 NewAnimSetIdx)
+{
+	// For derived classes to override.
+}
+
 void UContextualAnimSceneActorComponent::OnRep_RepTransitionSingleActor()
 {
 	UE_LOG(LogContextualAnim, Verbose, TEXT("%-21s UContextualAnimSceneActorComponent::OnRep_RepTransitionSingleActor Owner: %s Id: %d RepCounter: %d SectionIdx: %d AnimSetIdx: %d"),
@@ -580,6 +617,16 @@ void UContextualAnimSceneActorComponent::OnRep_RepTransitionSingleActor()
 					PlayAnimation_Internal(AnimTrack->Animation, 0.f, false);
 
 					AddOrUpdateWarpTargets(RepTransitionSingleActorData.SectionIdx, RepTransitionSingleActorData.AnimSetIdx, RepTransitionSingleActorData.WarpPoints, RepTransitionSingleActorData.ExternalWarpTargets);
+
+					if (UCharacterMovementComponent* MovementComp = OwnerBinding->GetCharacterMovementComponent())
+					{
+						if (MovementComp->MovementMode != AnimTrack->MovementMode)
+						{
+							MovementComp->SetMovementMode(AnimTrack->MovementMode);
+						}
+					}
+
+					OnTransitionSingleActor(*OwnerBinding, RepTransitionSingleActorData.SectionIdx, RepTransitionSingleActorData.AnimSetIdx);
 				}
 			}
 		}
@@ -738,14 +785,21 @@ bool UContextualAnimSceneActorComponent::ServerEarlyOutContextualAnimScene_Valid
 void UContextualAnimSceneActorComponent::OnRep_TransitionData()
 {
 	UE_LOG(LogContextualAnim, Verbose, TEXT("%-21s UContextualAnimSceneActorComponent::OnRep_TransitionData Actor: %s SectionIdx: %d AnimsetIdx: %d RepCounter: %d"),
-		*UEnum::GetValueAsString(TEXT("Engine.ENetRole"), GetOwner()->GetLocalRole()), *GetNameSafe(GetOwner()),
-		RepTransitionData.SectionIdx, RepTransitionData.AnimSetIdx, RepTransitionData.RepCounter);
+		*UEnum::GetValueAsString(TEXT("Engine.ENetRole"), GetOwner()->GetLocalRole()), *GetNameSafe(GetOwner()), RepTransitionData.SectionIdx, RepTransitionData.AnimSetIdx, RepTransitionData.RepCounter);
 
 	if (!RepTransitionData.IsValid())
 	{
 		return;
 	}
 
+	if (!Bindings.IsValid())
+	{
+		UE_LOG(LogContextualAnim, Warning, TEXT("%-21s UContextualAnimSceneActorComponent::OnRep_TransitionData Actor: %s Current bindings INVALID"),
+			*UEnum::GetValueAsString(TEXT("Engine.ENetRole"), GetOwner()->GetLocalRole()), *GetNameSafe(GetOwner()));
+
+		return;
+	}
+
 	HandleTransitionEveryone(RepTransitionData.SectionIdx, RepTransitionData.AnimSetIdx, RepTransitionData.WarpPoints, RepTransitionData.ExternalWarpTargets);
 }
 
@@ -834,7 +888,7 @@ void UContextualAnimSceneActorComponent::SetIgnoreCollisionWithOtherActors(bool
 	}
 }
 
-void UContextualAnimSceneActorComponent::SetCollisionState()
+void UContextualAnimSceneActorComponent::SetCollisionState(const FContextualAnimSceneBinding& Binding)
 {
 	if (const UContextualAnimSceneAsset* Asset = Bindings.GetSceneAsset())
 	{
@@ -847,19 +901,16 @@ void UContextualAnimSceneActorComponent::SetCollisionState()
 		{
 			if (UPrimitiveComponent* RootPrimitiveComponent = Cast<UPrimitiveComponent>(GetOwner()->GetRootComponent()))
 			{
-				if (const FContextualAnimSceneBinding* Binding = Bindings.FindBindingByActor(GetOwner()))
+				const TArray<TEnumAsByte<ECollisionChannel>>& ChannelsToIgnore = Asset->GetCollisionChannelsToIgnoreForRole(Bindings.GetRoleFromBinding(Binding));
+				if (ChannelsToIgnore.Num() > 0)
 				{
-					const TArray<TEnumAsByte<ECollisionChannel>>& ChannelsToIgnore = Asset->GetCollisionChannelsToIgnoreForRole(Bindings.GetRoleFromBinding(*Binding));
-					if (ChannelsToIgnore.Num() > 0)
+					CharacterPropertiesBackup.CollisionResponses.Reset(ChannelsToIgnore.Num());
+					for (ECollisionChannel Channel : ChannelsToIgnore)
 					{
-						CharacterPropertiesBackup.CollisionResponses.Reset(ChannelsToIgnore.Num());
-						for (ECollisionChannel Channel : ChannelsToIgnore)
-						{
-							ECollisionResponse Response = RootPrimitiveComponent->GetCollisionResponseToChannel(Channel);
-							CharacterPropertiesBackup.CollisionResponses.Add(MakeTuple(Channel, Response));
+						ECollisionResponse Response = RootPrimitiveComponent->GetCollisionResponseToChannel(Channel);
+						CharacterPropertiesBackup.CollisionResponses.Add(MakeTuple(Channel, Response));
 
-							RootPrimitiveComponent->SetCollisionResponseToChannel(Channel, ECR_Ignore);
-						}
+						RootPrimitiveComponent->SetCollisionResponseToChannel(Channel, ECR_Ignore);
 					}
 				}
 			}
@@ -867,7 +918,7 @@ void UContextualAnimSceneActorComponent::SetCollisionState()
 	}
 }
 
-void UContextualAnimSceneActorComponent::RestoreCollisionState()
+void UContextualAnimSceneActorComponent::RestoreCollisionState(const FContextualAnimSceneBinding& Binding)
 {
 	if (const UContextualAnimSceneAsset* Asset = Bindings.GetSceneAsset())
 	{
@@ -920,14 +971,21 @@ void UContextualAnimSceneActorComponent::JoinScene(const FContextualAnimSceneBin
 
 		AddOrUpdateWarpTargets(AnimTrack.SectionIdx, AnimTrack.AnimSetIdx, WarpPoints, ExternalWarpTargets);
 
-		SetCollisionState();
+		SetCollisionState(*Binding);
 
-		SetMovementState(AnimTrack.bRequireFlyingMode);
+		SetMovementState(*Binding, AnimTrack.MovementMode);
+
+		OnJoinScene(*Binding);
 
 		OnJoinedSceneDelegate.Broadcast(this);
 	}
 }
 
+void UContextualAnimSceneActorComponent::OnJoinScene(const FContextualAnimSceneBinding& Binding)
+{
+	// For derived classes to override.
+}
+
 void UContextualAnimSceneActorComponent::LeaveScene()
 {
 	if (const FContextualAnimSceneBinding* Binding = Bindings.FindBindingByActor(GetOwner()))
@@ -957,10 +1015,11 @@ void UContextualAnimSceneActorComponent::LeaveScene()
 			SkelMeshComp->OnTickPose.RemoveAll(this);
 		}
 		
-		RestoreCollisionState();
+		RestoreCollisionState(*Binding);
 
-		const FContextualAnimTrack& AnimTrack = Bindings.GetAnimTrackFromBinding(*Binding);
-		RestoreMovementState(AnimTrack.bRequireFlyingMode);
+		RestoreMovementState(*Binding);
+
+		OnLeaveScene(*Binding);
 
 		OnLeftSceneDelegate.Broadcast(this);
 
@@ -968,15 +1027,21 @@ void UContextualAnimSceneActorComponent::LeaveScene()
 	}
 }
 
-void UContextualAnimSceneActorComponent::SetMovementState(bool bRequireFlyingMode)
+void UContextualAnimSceneActorComponent::OnLeaveScene(const FContextualAnimSceneBinding& Binding)
 {
-	if (UCharacterMovementComponent* MovementComp = GetOwner()->FindComponentByClass<UCharacterMovementComponent>())
+	// For derived classes to override.
+}
+
+void UContextualAnimSceneActorComponent::SetMovementState(const FContextualAnimSceneBinding& Binding, EMovementMode DesiredMoveMode)
+{
+	if (UCharacterMovementComponent* MovementComp = Binding.GetCharacterMovementComponent())
 	{
 		// Save movement state before the interaction starts so we can restore it when it ends
 		CharacterPropertiesBackup.bIgnoreClientMovementErrorChecksAndCorrection = MovementComp->bIgnoreClientMovementErrorChecksAndCorrection;
 		CharacterPropertiesBackup.bAllowPhysicsRotationDuringAnimRootMotion = MovementComp->bAllowPhysicsRotationDuringAnimRootMotion;
 		CharacterPropertiesBackup.bUseControllerDesiredRotation = MovementComp->bUseControllerDesiredRotation;
 		CharacterPropertiesBackup.bOrientRotationToMovement = MovementComp->bOrientRotationToMovement;
+		CharacterPropertiesBackup.MovementMode = MovementComp->MovementMode;
 
 		// Disable movement correction.
 		MovementComp->bIgnoreClientMovementErrorChecksAndCorrection = true;
@@ -986,32 +1051,31 @@ void UContextualAnimSceneActorComponent::SetMovementState(bool bRequireFlyingMod
 		MovementComp->bUseControllerDesiredRotation = false;
 		MovementComp->bOrientRotationToMovement = false;
 
-		//@TODO: Temp solution that assumes these interactions are not locally predicted and that is ok to be in flying mode during the entire animation
-		if (bRequireFlyingMode && MovementComp->MovementMode != MOVE_Flying)
+		if (MovementComp->MovementMode != DesiredMoveMode)
 		{
-			MovementComp->SetMovementMode(MOVE_Flying);
+			MovementComp->SetMovementMode(DesiredMoveMode);
 		}
 	}
 }
 
-void UContextualAnimSceneActorComponent::RestoreMovementState(bool bRequireFlyingMode)
+void UContextualAnimSceneActorComponent::RestoreMovementState(const FContextualAnimSceneBinding& Binding)
 {
-	if (UCharacterMovementComponent* MovementComp = GetOwner()->FindComponentByClass<UCharacterMovementComponent>())
+	if (UCharacterMovementComponent* MovementComp = Binding.GetCharacterMovementComponent())
 	{
 		// Restore movement state
 		MovementComp->bIgnoreClientMovementErrorChecksAndCorrection = CharacterPropertiesBackup.bIgnoreClientMovementErrorChecksAndCorrection;
 		MovementComp->bAllowPhysicsRotationDuringAnimRootMotion = CharacterPropertiesBackup.bAllowPhysicsRotationDuringAnimRootMotion;
 		MovementComp->bUseControllerDesiredRotation = CharacterPropertiesBackup.bUseControllerDesiredRotation;
 		MovementComp->bOrientRotationToMovement = CharacterPropertiesBackup.bOrientRotationToMovement;
-
-		//@TODO: Temp solution that assumes these interactions are not locally predicted and that is ok to be in flying mode during the entire animation
-		if (bRequireFlyingMode && MovementComp->MovementMode == MOVE_Flying)
-		{
-			MovementComp->SetMovementMode(MOVE_Walking);
-		}
+		MovementComp->SetMovementMode(CharacterPropertiesBackup.MovementMode);
 	}
 }
 
+bool UContextualAnimSceneActorComponent::CanLeaveScene(const FContextualAnimSceneBinding& Binding)
+{
+	return true;
+}
+
 void UContextualAnimSceneActorComponent::OnMontageBlendingOut(UAnimMontage* Montage, bool bInterrupted)
 {
 	if (bGuardAnimEvents)
diff --git a/Engine/Plugins/Experimental/Animation/ContextualAnimation/Source/ContextualAnimation/Private/ContextualAnimSceneInstance.cpp b/Engine/Plugins/Experimental/Animation/ContextualAnimation/Source/ContextualAnimation/Private/ContextualAnimSceneInstance.cpp
--- a/Engine/Plugins/Experimental/Animation/ContextualAnimation/Source/ContextualAnimation/Private/ContextualAnimSceneInstance.cpp
+++ b/Engine/Plugins/Experimental/Animation/ContextualAnimation/Source/ContextualAnimation/Private/ContextualAnimSceneInstance.cpp
@@ -102,13 +102,9 @@ float UContextualAnimSceneInstance::Join(FContextualAnimSceneBinding& Binding)
 			}
 		}
 
-		//@TODO: Temp, until we have a way to switch between movement mode using AnimNotifyState
-		if (Bindings.GetAnimTrackFromBinding(Binding).bRequireFlyingMode)
+		if (UCharacterMovementComponent* CharacterMovementComp = Actor->FindComponentByClass<UCharacterMovementComponent>())
 		{
-			if (UCharacterMovementComponent* CharacterMovementComp = Actor->FindComponentByClass<UCharacterMovementComponent>())
-			{
-				CharacterMovementComp->SetMovementMode(MOVE_Flying);
-			}
+			CharacterMovementComp->SetMovementMode(Bindings.GetAnimTrackFromBinding(Binding).MovementMode);
 		}
 	}
 
@@ -290,12 +286,9 @@ void UContextualAnimSceneInstance::OnMontageBlendingOut(UAnimMontage* Montage, b
 					AnimInstance->OnPlayMontageNotifyEnd.RemoveDynamic(this, &UContextualAnimSceneInstance::OnNotifyEndReceived);
 					AnimInstance->OnMontageBlendingOut.RemoveDynamic(this, &UContextualAnimSceneInstance::OnMontageBlendingOut);
 
-					if (AnimTrack.bRequireFlyingMode)
+					if (UCharacterMovementComponent* CharacterMovementComp = Actor->FindComponentByClass<UCharacterMovementComponent>())
 					{
-						if (UCharacterMovementComponent* CharacterMovementComp = Actor->FindComponentByClass<UCharacterMovementComponent>())
-						{
-							CharacterMovementComp->SetMovementMode(MOVE_Walking);
-						}
+						CharacterMovementComp->SetMovementMode(MOVE_Walking);
 					}
 				}
 
diff --git a/Engine/Plugins/Experimental/Animation/ContextualAnimation/Source/ContextualAnimation/Private/ContextualAnimTypes.cpp b/Engine/Plugins/Experimental/Animation/ContextualAnimation/Source/ContextualAnimation/Private/ContextualAnimTypes.cpp
--- a/Engine/Plugins/Experimental/Animation/ContextualAnimation/Source/ContextualAnimation/Private/ContextualAnimTypes.cpp
+++ b/Engine/Plugins/Experimental/Animation/ContextualAnimation/Source/ContextualAnimation/Private/ContextualAnimTypes.cpp
@@ -11,6 +11,7 @@
 #include "ContextualAnimSceneActorComponent.h"
 #include "ContextualAnimSceneAsset.h"
 #include "Components/SkeletalMeshComponent.h"
+#include "GameFramework/CharacterMovementComponent.h"
 
 #include UE_INLINE_GENERATED_CPP_BY_NAME(ContextualAnimTypes)
 
@@ -280,6 +281,19 @@ USkeletalMeshComponent* FContextualAnimSceneBindingContext::GetSkeletalMeshCompo
 	return CachedSkeletalMesh.Get();
 }
 
+UCharacterMovementComponent* FContextualAnimSceneBindingContext::GetCharacterMovementComponent() const
+{
+	if (!CachedMovementComp.IsValid() || CachedMovementComp->GetOwner() != GetActor())
+	{
+		if (Actor.IsValid())
+		{
+			CachedMovementComp = Actor->FindComponentByClass<UCharacterMovementComponent>();
+		}
+	}
+
+	return CachedMovementComp.Get();
+}
+
 // FContextualAnimSceneBinding
 ///////////////////////////////////////////////////////////////////////
 
diff --git a/Engine/Plugins/Experimental/Animation/ContextualAnimation/Source/ContextualAnimation/Public/ContextualAnimSceneActorComponent.h b/Engine/Plugins/Experimental/Animation/ContextualAnimation/Source/ContextualAnimation/Public/ContextualAnimSceneActorComponent.h
--- a/Engine/Plugins/Experimental/Animation/ContextualAnimation/Source/ContextualAnimation/Public/ContextualAnimSceneActorComponent.h
+++ b/Engine/Plugins/Experimental/Animation/ContextualAnimation/Source/ContextualAnimation/Public/ContextualAnimSceneActorComponent.h
@@ -183,7 +183,7 @@ public:
 	/** Called when the actor owner of this component joins an scene */
 	UFUNCTION(BlueprintCallable, Category = "Contextual Anim|Scene Actor Component")
 	void OnJoinedScene(const FContextualAnimSceneBindings& InBindings);
-	
+
 	/** Called from the scene instance when the actor owner of this component leave an scene */
 	UFUNCTION(BlueprintCallable, Category = "Contextual Anim|Scene Actor Component")
 	void OnLeftScene();
@@ -213,6 +213,7 @@ public:
 	bool TransitionSingleActor(int32 SectionIdx, const TArray<FContextualAnimWarpTarget>& ExternalWarpTargets);
 	bool TransitionSingleActor(int32 SectionIdx, int32 AnimSetIdx, const TArray<FContextualAnimWarpTarget>& ExternalWarpTargets);
 
+	UFUNCTION(BlueprintCallable, Category = "Contextual Anim|Scene Actor Component")
 	void EarlyOutContextualAnimScene();
 
 	bool IsOwnerLocallyControlled() const;
@@ -252,6 +253,7 @@ protected:
 		bool bAllowPhysicsRotationDuringAnimRootMotion = false;
 		bool bUseControllerDesiredRotation = false;
 		bool bOrientRotationToMovement = false;
+		EMovementMode MovementMode = EMovementMode::MOVE_Walking;
 		TArray<TTuple<ECollisionChannel, ECollisionResponse>> CollisionResponses;
 	};
 	FCharacterProperties CharacterPropertiesBackup;
@@ -278,11 +280,12 @@ protected:
 	void OnRep_TransitionData();
 
 	void SetIgnoreCollisionWithOtherActors(bool bValue) const;
-	void SetCollisionState();
-	void RestoreCollisionState();
+	
+	virtual void SetCollisionState(const FContextualAnimSceneBinding& Binding);
+	virtual void RestoreCollisionState(const FContextualAnimSceneBinding& Binding);
 
-	void SetMovementState(bool bRequireFlyingMode);
-	void RestoreMovementState(bool bRequireFlyingMode);
+	void SetMovementState(const FContextualAnimSceneBinding& Binding, EMovementMode DesiredMoveMode);
+	void RestoreMovementState(const FContextualAnimSceneBinding& Binding);
 
 	UFUNCTION()
 	void OnMontageBlendingOut(UAnimMontage* Montage, bool bInterrupted);
@@ -295,10 +298,18 @@ protected:
 
 	void JoinScene(const FContextualAnimSceneBindings& InBindings, const TArray<FContextualAnimWarpPoint> WarpPoints, const TArray<FContextualAnimWarpTarget>& ExternalWarpTargets);
 
+	virtual void OnJoinScene(const FContextualAnimSceneBinding& Binding);
+
 	void LeaveScene();
 
+	virtual bool CanLeaveScene(const FContextualAnimSceneBinding& Binding);
+
+	virtual void OnLeaveScene(const FContextualAnimSceneBinding& Binding);
+
 	void LateJoinScene(const FContextualAnimSceneBindings& InBindings, int32 SectionIdx, int32 AnimSetIdx, const TArray<FContextualAnimWarpPoint>& WarpPoints, const TArray<FContextualAnimWarpTarget>& ExternalWarpTargets);
 
+	virtual void OnLateJoinScene(const FContextualAnimSceneBinding& Binding, int32 SectionIdx, int32 AnimSetIdx);
+
 	UFUNCTION(Server, Reliable, WithValidation)
 	void ServerStartContextualAnimScene(const FContextualAnimSceneBindings& InBindings);
 
@@ -313,6 +324,10 @@ protected:
 
 	void HandleTransitionEveryone(int32 NewSectionIdx, int32 NewAnimSetIdx, const TArray<FContextualAnimWarpPoint>& WarpPoints, const TArray<FContextualAnimWarpTarget>& ExternalWarpTargets);
 
+	virtual void OnTransitionScene(const FContextualAnimSceneBinding& Binding, int32 SectionIdx, int32 AnimSetIdx);
+
+	virtual void OnTransitionSingleActor(const FContextualAnimSceneBinding& Binding, int32 SectionIdx, int32 AnimSetIdx);
+
 private:
 
 	bool bRegistered = false;
diff --git a/Engine/Plugins/Experimental/Animation/ContextualAnimation/Source/ContextualAnimation/Public/ContextualAnimSceneAsset.h b/Engine/Plugins/Experimental/Animation/ContextualAnimation/Source/ContextualAnimation/Public/ContextualAnimSceneAsset.h
--- a/Engine/Plugins/Experimental/Animation/ContextualAnimation/Source/ContextualAnimation/Public/ContextualAnimSceneAsset.h
+++ b/Engine/Plugins/Experimental/Animation/ContextualAnimation/Source/ContextualAnimation/Public/ContextualAnimSceneAsset.h
@@ -22,13 +22,28 @@ struct FContextualAnimIgnoreChannelsParam
 {
 	GENERATED_BODY()
 
-	UPROPERTY(EditAnywhere, Category = "Defaults")
+	UPROPERTY(EditAnywhere, Category = "Defaults", meta = (GetOptions = "GetRoles"))
 	FName Role = NAME_None;
 
 	UPROPERTY(EditAnywhere, Category = "Defaults")
 	TArray<TEnumAsByte<ECollisionChannel>> Channels;
 };
 
+USTRUCT()
+struct FContextualAnimAttachmentParams
+{
+	GENERATED_BODY()
+
+	UPROPERTY(EditAnywhere, Category = "Defaults", meta = (GetOptions = "GetRoles"))
+	FName Role = NAME_None;
+
+	UPROPERTY(EditAnywhere, Category = "Defaults")
+	FName SocketName = NAME_None;
+
+	UPROPERTY(EditAnywhere, Category = "Defaults")
+	FTransform RelativeTransform = FTransform::Identity;
+};
+
 UCLASS(Blueprintable)
 class CONTEXTUALANIMATION_API UContextualAnimRolesAsset : public UDataAsset
 {
@@ -222,6 +237,11 @@ public:
 	FORCEINLINE bool ShouldPrecomputeAlignmentTracks() const { return bPrecomputeAlignmentTracks; }
 
 	const TArray<TEnumAsByte<ECollisionChannel>>& GetCollisionChannelsToIgnoreForRole(FName Role) const;
+	
+	const FContextualAnimAttachmentParams* GetAttachmentParamsForRole(FName Role) const
+	{
+		return AttachmentParams.FindByPredicate([Role](const FContextualAnimAttachmentParams& Item) { return Item.Role == Role; });
+	}
 
 	bool HasValidData() const { return RolesAsset != nullptr && Sections.Num() > 0 && Sections[0].AnimSets.Num() > 0; }
 
@@ -325,6 +345,9 @@ protected:
 	UPROPERTY(EditAnywhere, Category = "Settings", meta = (EditCondition = "CollisionBehavior==EContextualAnimCollisionBehavior::IgnoreChannels", EditConditionHides))
 	TArray<FContextualAnimIgnoreChannelsParam> CollisionChannelsToIgnoreParams;
 
+	UPROPERTY(EditAnywhere, Category = "Settings")
+	TArray<FContextualAnimAttachmentParams> AttachmentParams;
+
 	/** Whether we should extract and cache alignment tracks off line. */
 	UPROPERTY(EditAnywhere, Category = "Settings", AdvancedDisplay)
 	bool bPrecomputeAlignmentTracks = false;
diff --git a/Engine/Plugins/Experimental/Animation/ContextualAnimation/Source/ContextualAnimation/Public/ContextualAnimTypes.h b/Engine/Plugins/Experimental/Animation/ContextualAnimation/Source/ContextualAnimation/Public/ContextualAnimTypes.h
--- a/Engine/Plugins/Experimental/Animation/ContextualAnimation/Source/ContextualAnimation/Public/ContextualAnimTypes.h
+++ b/Engine/Plugins/Experimental/Animation/ContextualAnimation/Source/ContextualAnimation/Public/ContextualAnimTypes.h
@@ -17,6 +17,7 @@ class UAnimSequenceBase;
 class UContextualAnimSelectionCriterion;
 class UContextualAnimSceneAsset;
 class UContextualAnimSceneActorComponent;
+class UCharacterMovementComponent;
 struct FAnimMontageInstance;
 
 namespace UE 
@@ -73,7 +74,7 @@ struct CONTEXTUALANIMATION_API FContextualAnimTrack
 	float AnimMaxStartTime = 0.f;
 
 	UPROPERTY(EditAnywhere, BlueprintReadWrite, Category = "Defaults")
-	bool bRequireFlyingMode = false;
+	TEnumAsByte<enum EMovementMode> MovementMode = EMovementMode::MOVE_Walking;
 
 	/** Whether the actor that should play this animation is optional */
 	UPROPERTY(EditAnywhere, BlueprintReadWrite, Category = "Defaults")
@@ -340,6 +341,8 @@ struct CONTEXTUALANIMATION_API FContextualAnimSceneBindingContext
 
 	UContextualAnimSceneActorComponent* GetSceneActorComponent() const;
 
+	UCharacterMovementComponent* GetCharacterMovementComponent() const;
+
 	void SetExternalTransform(const FTransform& InTransform);
 
 	FTransform GetTransform() const;
@@ -370,6 +373,9 @@ private:
 	UPROPERTY(NotReplicated)
 	mutable TWeakObjectPtr<USkeletalMeshComponent> CachedSkeletalMesh = nullptr;
 
+	UPROPERTY(NotReplicated)
+	mutable TWeakObjectPtr<UCharacterMovementComponent> CachedMovementComp = nullptr;
+
 	TOptional<FTransform> ExternalTransform;
 
 	TOptional<FVector> ExternalVelocity;
@@ -394,6 +400,7 @@ struct CONTEXTUALANIMATION_API FContextualAnimSceneBinding
 	FORCEINLINE UAnimInstance* GetAnimInstance() const { return Context.GetAnimInstance(); }
 	FORCEINLINE USkeletalMeshComponent* GetSkeletalMeshComponent() const { return Context.GetSkeletalMeshComponent(); }
 	FORCEINLINE UContextualAnimSceneActorComponent* GetSceneActorComponent() const { return Context.GetSceneActorComponent(); }
+	FORCEINLINE UCharacterMovementComponent* GetCharacterMovementComponent() const { return Context.GetCharacterMovementComponent(); }
 	FORCEINLINE int32 GetAnimTrackIdx() const { return AnimTrackIdx; }
 	
 	void SetAnimTrack(const FContextualAnimTrack& InAnimTrack);
diff --git a/Engine/Plugins/Experimental/Animation/ContextualAnimation/Source/ContextualAnimationEditor/Private/ContextualAnimViewModel.cpp b/Engine/Plugins/Experimental/Animation/ContextualAnimation/Source/ContextualAnimationEditor/Private/ContextualAnimViewModel.cpp
--- a/Engine/Plugins/Experimental/Animation/ContextualAnimation/Source/ContextualAnimationEditor/Private/ContextualAnimViewModel.cpp
+++ b/Engine/Plugins/Experimental/Animation/ContextualAnimation/Source/ContextualAnimationEditor/Private/ContextualAnimViewModel.cpp
@@ -285,8 +285,7 @@ AActor* FContextualAnimViewModel::SpawnPreviewActor(const FContextualAnimTrack&
 				CharacterMovementComp->bUseControllerDesiredRotation = false;
 				CharacterMovementComp->RotationRate = FRotator(0.f, 540.0, 0.f);
 				CharacterMovementComp->bRunPhysicsWithNoController = true;
-
-				CharacterMovementComp->SetMovementMode(AnimTrack.bRequireFlyingMode ? EMovementMode::MOVE_Flying : EMovementMode::MOVE_Walking);
+				CharacterMovementComp->SetMovementMode(AnimTrack.MovementMode);
 			}
 
 			UMotionWarpingComponent* MotionWarpingComp = NewObject<UMotionWarpingComponent>(PreviewCharacter);
@@ -575,7 +574,7 @@ void FContextualAnimViewModel::AddNewAnimSet(const FContextualAnimNewAnimSetPara
 		FContextualAnimTrack AnimTrack;
 		AnimTrack.Role = Data.RoleName;
 		AnimTrack.Animation = Data.Animation;
-		AnimTrack.bRequireFlyingMode = Data.bRequiresFlyingMode;
+		AnimTrack.MovementMode = Data.MovementMode;
 		AnimTrack.bOptional = Data.bOptional;
 		AnimSet.Tracks.Add(AnimTrack);
 		AnimSet.RandomWeight = Params.RandomWeight;
diff --git a/Engine/Plugins/Experimental/Animation/ContextualAnimation/Source/ContextualAnimationEditor/Public/ContextualAnimEditorTypes.h b/Engine/Plugins/Experimental/Animation/ContextualAnimation/Source/ContextualAnimationEditor/Public/ContextualAnimEditorTypes.h
--- a/Engine/Plugins/Experimental/Animation/ContextualAnimation/Source/ContextualAnimationEditor/Public/ContextualAnimEditorTypes.h
+++ b/Engine/Plugins/Experimental/Animation/ContextualAnimation/Source/ContextualAnimationEditor/Public/ContextualAnimEditorTypes.h
@@ -75,7 +75,7 @@ struct FContextualAnimNewAnimSetData
 	TObjectPtr<class UAnimMontage> Animation = nullptr;
 
 	UPROPERTY(EditAnywhere, Category = "Settings")
-	bool bRequiresFlyingMode = false;
+	TEnumAsByte<enum EMovementMode> MovementMode = EMovementMode::MOVE_Walking;
 
 	UPROPERTY(EditAnywhere, Category = "Settings")
 	bool bOptional = false;
