Contextual Anim: Added blueprint callable helper function to know if the actor is in an active interaction

diff --git a/Engine/Plugins/Experimental/Animation/ContextualAnimation/Source/ContextualAnimation/Private/ContextualAnimSceneActorComponent.cpp b/Engine/Plugins/Experimental/Animation/ContextualAnimation/Source/ContextualAnimation/Private/ContextualAnimSceneActorComponent.cpp
--- a/Engine/Plugins/Experimental/Animation/ContextualAnimation/Source/ContextualAnimation/Private/ContextualAnimSceneActorComponent.cpp
+++ b/Engine/Plugins/Experimental/Animation/ContextualAnimation/Source/ContextualAnimation/Private/ContextualAnimSceneActorComponent.cpp
@@ -111,6 +111,11 @@ bool UContextualAnimSceneActorComponent::IsOwnerLocallyControlled() const
 	return false;
 }
 
+bool UContextualAnimSceneActorComponent::IsInActiveScene() const
+{
+	return (Bindings.IsValid() && Bindings.FindBindingByActor(GetOwner()) != nullptr);
+}
+
 void UContextualAnimSceneActorComponent::PlayAnimation_Internal(UAnimSequenceBase* Animation, float StartTime, bool bSyncPlaybackTime)
 {
 	// Replaced TGuardValue with this one frame delay because for some reason, apparently random (needs more investigation), in standalone OnMontageBlendingOut event is queued instead of triggered inline, 
diff --git a/Engine/Plugins/Experimental/Animation/ContextualAnimation/Source/ContextualAnimation/Public/ContextualAnimSceneActorComponent.h b/Engine/Plugins/Experimental/Animation/ContextualAnimation/Source/ContextualAnimation/Public/ContextualAnimSceneActorComponent.h
--- a/Engine/Plugins/Experimental/Animation/ContextualAnimation/Source/ContextualAnimation/Public/ContextualAnimSceneActorComponent.h
+++ b/Engine/Plugins/Experimental/Animation/ContextualAnimation/Source/ContextualAnimation/Public/ContextualAnimSceneActorComponent.h
@@ -180,6 +180,9 @@ public:
 
 	const FContextualAnimSceneBindings& GetBindings() const { return Bindings; };
 
+	UFUNCTION(BlueprintPure, Category = "Contextual Anim|Scene Actor Component")
+	bool IsInActiveScene() const;
+
 	/** Called when the actor owner of this component joins an scene */
 	UFUNCTION(BlueprintCallable, Category = "Contextual Anim|Scene Actor Component")
 	void OnJoinedScene(const FContextualAnimSceneBindings& InBindings);
