Contextual Anim: Removed deprecated scene instance and manager

diff --git a/Engine/Plugins/Experimental/Animation/ContextualAnimation/Source/ContextualAnimation/Private/ContextualAnimManager.cpp b/Engine/Plugins/Experimental/Animation/ContextualAnimation/Source/ContextualAnimation/Private/ContextualAnimManager.cpp
deleted file mode 100644
diff --git a/Engine/Plugins/Experimental/Animation/ContextualAnimation/Source/ContextualAnimation/Private/ContextualAnimSceneActorComponent.cpp b/Engine/Plugins/Experimental/Animation/ContextualAnimation/Source/ContextualAnimation/Private/ContextualAnimSceneActorComponent.cpp
--- a/Engine/Plugins/Experimental/Animation/ContextualAnimation/Source/ContextualAnimation/Private/ContextualAnimSceneActorComponent.cpp
+++ b/Engine/Plugins/Experimental/Animation/ContextualAnimation/Source/ContextualAnimation/Private/ContextualAnimSceneActorComponent.cpp
@@ -7,7 +7,6 @@
 #include "GameFramework/CharacterMovementComponent.h"
 #include "ContextualAnimSelectionCriterion.h"
 #include "Components/SkeletalMeshComponent.h"
-#include "ContextualAnimManager.h"
 #include "ContextualAnimSceneAsset.h"
 #include "ContextualAnimUtilities.h"
 #include "AnimNotifyState_IKWindow.h"
@@ -852,30 +851,6 @@ FBoxSphereBounds UContextualAnimSceneActorComponent::CalcBounds(const FTransform
 	return FBoxSphereBounds(FSphere(GetComponentTransform().GetLocation(), Radius));
 }
 
-void UContextualAnimSceneActorComponent::OnRegister()
-{
-	Super::OnRegister();
-
-	UContextualAnimManager* ContextAnimManager = UContextualAnimManager::Get(GetWorld());
-	if (ensure(!bRegistered) && ContextAnimManager)
-	{
-		ContextAnimManager->RegisterSceneActorComponent(this);
-		bRegistered = true;
-	}
-}
-
-void UContextualAnimSceneActorComponent::OnUnregister()
-{
-	Super::OnUnregister();
-
-	UContextualAnimManager* ContextAnimManager = UContextualAnimManager::Get(GetWorld());
-	if (bRegistered && ContextAnimManager)
-	{
-		ContextAnimManager->UnregisterSceneActorComponent(this);
-		bRegistered = false;
-	}
-}
-
 void UContextualAnimSceneActorComponent::SetIgnoreCollisionWithOtherActors(bool bValue) const
 {
 	const AActor* OwnerActor = GetOwner();
diff --git a/Engine/Plugins/Experimental/Animation/ContextualAnimation/Source/ContextualAnimation/Private/ContextualAnimSceneInstance.cpp b/Engine/Plugins/Experimental/Animation/ContextualAnimation/Source/ContextualAnimation/Private/ContextualAnimSceneInstance.cpp
deleted file mode 100644
diff --git a/Engine/Plugins/Experimental/Animation/ContextualAnimation/Source/ContextualAnimation/Private/ContextualAnimTransition.cpp b/Engine/Plugins/Experimental/Animation/ContextualAnimation/Source/ContextualAnimation/Private/ContextualAnimTransition.cpp
deleted file mode 100644
diff --git a/Engine/Plugins/Experimental/Animation/ContextualAnimation/Source/ContextualAnimation/Private/ContextualAnimation.cpp b/Engine/Plugins/Experimental/Animation/ContextualAnimation/Source/ContextualAnimation/Private/ContextualAnimation.cpp
--- a/Engine/Plugins/Experimental/Animation/ContextualAnimation/Source/ContextualAnimation/Private/ContextualAnimation.cpp
+++ b/Engine/Plugins/Experimental/Animation/ContextualAnimation/Source/ContextualAnimation/Private/ContextualAnimation.cpp
@@ -1,46 +1,16 @@
 // Copyright Epic Games, Inc. All Rights Reserved.
 
 #include "ContextualAnimation.h"
-#include "ContextualAnimManager.h"
 #include "Modules/ModuleManager.h"
 
 #define LOCTEXT_NAMESPACE "ContextualAnimationModule"
 
-TMap<const UWorld*, TObjectPtr<UContextualAnimManager>> FContextualAnimationModule::WorldToManagerMap;
-
 void FContextualAnimationModule::StartupModule()
 {
-	OnPreWorldInitDelegateHandle = FWorldDelegates::OnPreWorldInitialization.AddStatic(&FContextualAnimationModule::OnWorldInit);
-	OnPostWorldCleanupDelegateHandle = FWorldDelegates::OnPostWorldCleanup.AddStatic(&FContextualAnimationModule::OnWorldCleanup);
 }
 
 void FContextualAnimationModule::ShutdownModule()
 {
-	FWorldDelegates::OnPreWorldInitialization.Remove(OnPreWorldInitDelegateHandle);
-	FWorldDelegates::OnPostWorldCleanup.Remove(OnPostWorldCleanupDelegateHandle);
-
-	WorldToManagerMap.Empty();
-}
-
-void FContextualAnimationModule::AddReferencedObjects(FReferenceCollector& Collector)
-{
-	for (auto& Pair : WorldToManagerMap)
-	{
-		Collector.AddReferencedObject(Pair.Value, Pair.Key);
-	}
-}
-
-void FContextualAnimationModule::OnWorldInit(UWorld* World, const UWorld::InitializationValues IVS)
-{
-	if(World && World->IsGameWorld())
-	{
-		WorldToManagerMap.Add(World, NewObject<UContextualAnimManager>(World));
-	}
-}
-
-void FContextualAnimationModule::OnWorldCleanup(UWorld* World, bool bSessionEnded, bool bCleanupResources)
-{
-	WorldToManagerMap.Remove(World);
 }
 
 #undef LOCTEXT_NAMESPACE
diff --git a/Engine/Plugins/Experimental/Animation/ContextualAnimation/Source/ContextualAnimation/Public/ContextualAnimManager.h b/Engine/Plugins/Experimental/Animation/ContextualAnimation/Source/ContextualAnimation/Public/ContextualAnimManager.h
deleted file mode 100644
diff --git a/Engine/Plugins/Experimental/Animation/ContextualAnimation/Source/ContextualAnimation/Public/ContextualAnimSceneActorComponent.h b/Engine/Plugins/Experimental/Animation/ContextualAnimation/Source/ContextualAnimation/Public/ContextualAnimSceneActorComponent.h
--- a/Engine/Plugins/Experimental/Animation/ContextualAnimation/Source/ContextualAnimation/Public/ContextualAnimSceneActorComponent.h
+++ b/Engine/Plugins/Experimental/Animation/ContextualAnimation/Source/ContextualAnimation/Public/ContextualAnimSceneActorComponent.h
@@ -171,8 +171,6 @@ public:
 
 	static void AddReferencedObjects(UObject* InThis, FReferenceCollector& Collector);
 	virtual void GetLifetimeReplicatedProps(TArray< FLifetimeProperty >& OutLifetimeProps) const;
-	virtual void OnRegister() override;
-	virtual void OnUnregister() override;
 	virtual FBoxSphereBounds CalcBounds(const FTransform& LocalToWorld) const;
 	virtual FPrimitiveSceneProxy* CreateSceneProxy() override;
 
@@ -332,8 +330,6 @@ protected:
 	virtual void OnTransitionSingleActor(const FContextualAnimSceneBinding& Binding, int32 SectionIdx, int32 AnimSetIdx);
 
 private:
-
-	bool bRegistered = false;
 	bool bGuardAnimEvents = false;
 };
 
diff --git a/Engine/Plugins/Experimental/Animation/ContextualAnimation/Source/ContextualAnimation/Public/ContextualAnimSceneAsset.h b/Engine/Plugins/Experimental/Animation/ContextualAnimation/Source/ContextualAnimation/Public/ContextualAnimSceneAsset.h
--- a/Engine/Plugins/Experimental/Animation/ContextualAnimation/Source/ContextualAnimation/Public/ContextualAnimSceneAsset.h
+++ b/Engine/Plugins/Experimental/Animation/ContextualAnimation/Source/ContextualAnimation/Public/ContextualAnimSceneAsset.h
@@ -231,7 +231,6 @@ public:
 
 	FORCEINLINE const FName& GetPrimaryRole() const { return PrimaryRole; }
 	FORCEINLINE EContextualAnimCollisionBehavior GetCollisionBehavior() const { return CollisionBehavior; }
-	FORCEINLINE const TSubclassOf<UContextualAnimSceneInstance>& GetSceneInstanceClass() const { return SceneInstanceClass; }
 	FORCEINLINE int32 GetSampleRate() const { return SampleRate; }
 	FORCEINLINE float GetRadius() const { return Radius; }
 	FORCEINLINE bool ShouldPrecomputeAlignmentTracks() const { return bPrecomputeAlignmentTracks; }
@@ -336,9 +335,6 @@ protected:
 	UPROPERTY(EditAnywhere, Category = "Settings")
 	float Radius = 0.f;
 
-	UPROPERTY(EditAnywhere, Category = "Settings")
-	TSubclassOf<UContextualAnimSceneInstance> SceneInstanceClass;
-
 	UPROPERTY(EditAnywhere, Category = "Settings")
 	EContextualAnimCollisionBehavior CollisionBehavior = EContextualAnimCollisionBehavior::IgnoreActorWhenMoving;
 
diff --git a/Engine/Plugins/Experimental/Animation/ContextualAnimation/Source/ContextualAnimation/Public/ContextualAnimSceneInstance.h b/Engine/Plugins/Experimental/Animation/ContextualAnimation/Source/ContextualAnimation/Public/ContextualAnimSceneInstance.h
deleted file mode 100644
diff --git a/Engine/Plugins/Experimental/Animation/ContextualAnimation/Source/ContextualAnimation/Public/ContextualAnimTransition.h b/Engine/Plugins/Experimental/Animation/ContextualAnimation/Source/ContextualAnimation/Public/ContextualAnimTransition.h
deleted file mode 100644
diff --git a/Engine/Plugins/Experimental/Animation/ContextualAnimation/Source/ContextualAnimation/Public/ContextualAnimation.h b/Engine/Plugins/Experimental/Animation/ContextualAnimation/Source/ContextualAnimation/Public/ContextualAnimation.h
--- a/Engine/Plugins/Experimental/Animation/ContextualAnimation/Source/ContextualAnimation/Public/ContextualAnimation.h
+++ b/Engine/Plugins/Experimental/Animation/ContextualAnimation/Source/ContextualAnimation/Public/ContextualAnimation.h
@@ -5,39 +5,13 @@
 #include "Engine/World.h"
 #include "Modules/ModuleInterface.h"
 
-class UContextualAnimManager;
-
-class FContextualAnimationModule : public IModuleInterface, public FGCObject
+class FContextualAnimationModule : public IModuleInterface
 {
 public:
 
 	/** IModuleInterface implementation */
 	virtual void StartupModule() override;
 	virtual void ShutdownModule() override;
-
-	// Begin FGCObject overrides
-	virtual void AddReferencedObjects(FReferenceCollector& Collector) override;
-	virtual FString GetReferencerName() const override
-	{
-		return TEXT("FContextualAnimationModule");
-	}
-	// End FGCObject overrides
-
-	FORCEINLINE static UContextualAnimManager* GetManager(const UWorld* World)
-	{
-		return WorldToManagerMap.FindRef(World);
-	}
-
-private:
-
-	static void OnWorldInit(UWorld* World, const UWorld::InitializationValues IVS);
-
-	static void OnWorldCleanup(UWorld* World, bool bSessionEnded, bool bCleanupResources);
-
-	static TMap<const UWorld*, TObjectPtr<UContextualAnimManager>> WorldToManagerMap;
-
-	FDelegateHandle OnPreWorldInitDelegateHandle;
-	FDelegateHandle OnPostWorldCleanupDelegateHandle;
 };
 
 #if UE_ENABLE_INCLUDE_ORDER_DEPRECATED_IN_5_2
