Contextual Anim: Refactored guarding mechanism used to prevent OnMontageBlendingOut event from triggering during a transition within the same interaction to fix bug (apparently) caused by URO delaying the triggering of this event  + upgraded verbosity of some  logs.

diff --git a/Engine/Plugins/Experimental/Animation/ContextualAnimation/Source/ContextualAnimation/Private/ContextualAnimSceneActorComponent.cpp b/Engine/Plugins/Experimental/Animation/ContextualAnimation/Source/ContextualAnimation/Private/ContextualAnimSceneActorComponent.cpp
--- a/Engine/Plugins/Experimental/Animation/ContextualAnimation/Source/ContextualAnimation/Private/ContextualAnimSceneActorComponent.cpp
+++ b/Engine/Plugins/Experimental/Animation/ContextualAnimation/Source/ContextualAnimation/Private/ContextualAnimSceneActorComponent.cpp
@@ -117,28 +117,17 @@ bool UContextualAnimSceneActorComponent::IsInActiveScene() const
 
 void UContextualAnimSceneActorComponent::PlayAnimation_Internal(UAnimSequenceBase* Animation, float StartTime, bool bSyncPlaybackTime)
 {
-	// Replaced TGuardValue with this one frame delay because for some reason, apparently random (needs more investigation), in standalone OnMontageBlendingOut event is queued instead of triggered inline, 
-	// causing this guarding mechanism to fail because by the time the event triggers TGuardValue goes out of the scope
-	// Making our OnMontageBlendingOut to think the animation has been interrupted by an external system and forcing the actor to leave the interaction.
-	bGuardAnimEvents = true;
-	if (UWorld* World = GetWorld())
-	{
-		World->GetTimerManager().SetTimerForNextTick([WeakThis = MakeWeakObjectPtr(this)]()
-		{
-			if (UContextualAnimSceneActorComponent* Comp = WeakThis.Get())
-			{
-				Comp->bGuardAnimEvents = false;
-			}
-		});
-	}
-
 	if (UAnimInstance* AnimInstance = UContextualAnimUtilities::TryGetAnimInstance(GetOwner()))
 	{
-		UE_LOG(LogContextualAnim, VeryVerbose, TEXT("%-21s \t\tUContextualAnimSceneActorComponent::PlayAnimation_Internal Playing Animation. Actor: %s Anim: %s StartTime: %f bSyncPlaybackTime: %d"),
+		UE_LOG(LogContextualAnim, Log, TEXT("%-21s \t\tUContextualAnimSceneActorComponent::PlayAnimation_Internal Playing Animation. Actor: %s Anim: %s StartTime: %f bSyncPlaybackTime: %d"),
 			*UEnum::GetValueAsString(TEXT("Engine.ENetRole"), GetOwner()->GetLocalRole()), *GetNameSafe(GetOwner()), *GetNameSafe(Animation), StartTime, bSyncPlaybackTime);
 
 		//@TODO: Add support for dynamic montage
 		UAnimMontage* AnimMontage = Cast<UAnimMontage>(Animation);
+
+		// Keep track of this animation. Used as guarding mechanism in OnMonageBlendingOut to decide if is safe to leave the scene
+		AnimsPlayed.Add(AnimMontage);
+
 		AnimInstance->Montage_Play(AnimMontage, 1.f, EMontagePlayReturnType::MontageLength, StartTime);
 
 		AnimInstance->OnMontageBlendingOut.AddUniqueDynamic(this, &UContextualAnimSceneActorComponent::OnMontageBlendingOut);
@@ -254,7 +243,7 @@ bool UContextualAnimSceneActorComponent::LateJoinContextualAnimScene(AActor* Act
 		}
 	}
 
-	UE_LOG(LogContextualAnim, Verbose, TEXT("%-21s UContextualAnimSceneActorComponent::LateJoinContextualAnimScene Owner: %s Bindings Id: %d Section: %d Asset: %s. Requester: %s Role: %s"),
+	UE_LOG(LogContextualAnim, Log, TEXT("%-21s UContextualAnimSceneActorComponent::LateJoinContextualAnimScene Owner: %s Bindings Id: %d Section: %d Asset: %s. Requester: %s Role: %s"),
 		*UEnum::GetValueAsString(TEXT("Engine.ENetRole"), GetOwner()->GetLocalRole()), *GetNameSafe(GetOwner()), Bindings.GetID(), Bindings.GetSectionIdx(), *GetNameSafe(Bindings.GetSceneAsset()), *GetNameSafe(Actor), *Role.ToString());
 
 	// Add actor to the bindings
@@ -318,7 +307,7 @@ void UContextualAnimSceneActorComponent::LateJoinScene(const FContextualAnimScen
 
 	if (const FContextualAnimSceneBinding* Binding = InBindings.FindBindingByActor(GetOwner()))
 	{
-		UE_LOG(LogContextualAnim, Verbose, TEXT("%-21s UContextualAnimSceneActorComponent::LateJoinScene Actor: %s Role: %s Bindings Id: %d Section: %d Asset: %s"),
+		UE_LOG(LogContextualAnim, Log, TEXT("%-21s UContextualAnimSceneActorComponent::LateJoinScene Actor: %s Role: %s Bindings Id: %d Section: %d Asset: %s"),
 			*UEnum::GetValueAsString(TEXT("Engine.ENetRole"), GetOwner()->GetLocalRole()), *GetNameSafe(GetOwner()), *InBindings.GetRoleFromBinding(*Binding).ToString(), InBindings.GetID(), InBindings.GetSectionIdx(), *GetNameSafe(InBindings.GetSceneAsset()));
 
 		Bindings = InBindings;
@@ -347,7 +336,7 @@ void UContextualAnimSceneActorComponent::OnRep_LateJoinData()
 {
 	// This is received by the leader of the interaction on every remote client
 
-	UE_LOG(LogContextualAnim, Verbose, TEXT("%-21s UContextualAnimSceneActorComponent::OnRep_LateJoinData Owner: %s Bindings Id: %d Section: %d Asset: %s. Requester: %s Role: %s RepCounter: %d"),
+	UE_LOG(LogContextualAnim, Log, TEXT("%-21s UContextualAnimSceneActorComponent::OnRep_LateJoinData Owner: %s Bindings Id: %d Section: %d Asset: %s. Requester: %s Role: %s RepCounter: %d"),
 		*UEnum::GetValueAsString(TEXT("Engine.ENetRole"), GetOwner()->GetLocalRole()), *GetNameSafe(GetOwner()), Bindings.GetID(), Bindings.GetSectionIdx(), *GetNameSafe(Bindings.GetSceneAsset()), *GetNameSafe(RepLateJoinData.Actor), *RepLateJoinData.Role.ToString(), RepLateJoinData.RepCounter);
 
 	if (!RepLateJoinData.IsValid())
@@ -457,7 +446,7 @@ bool UContextualAnimSceneActorComponent::TransitionContextualAnimScene(FName Sec
 	const int32 SectionIdx = Bindings.GetSceneAsset()->GetSectionIndex(SectionName);
 	if (SectionIdx == INDEX_NONE)
 	{
-		UE_LOG(LogContextualAnim, Log, TEXT("%-21s UContextualAnimSceneActorComponent::TransitionContextualAnimScene. Invalid SectionName. Actor: %s SectionName: %s"),
+		UE_LOG(LogContextualAnim, Warning, TEXT("%-21s UContextualAnimSceneActorComponent::TransitionContextualAnimScene. Invalid SectionName. Actor: %s SectionName: %s"),
 			*UEnum::GetValueAsString(TEXT("Engine.ENetRole"), GetOwner()->GetLocalRole()), *GetNameSafe(GetOwner()), *SectionName.ToString());
 
 		return false;
@@ -466,7 +455,7 @@ bool UContextualAnimSceneActorComponent::TransitionContextualAnimScene(FName Sec
 	const int32 AnimSetIdx = Bindings.FindAnimSetForTransitionTo(SectionIdx);
 	if (AnimSetIdx == INDEX_NONE)
 	{
-		UE_LOG(LogContextualAnim, Log, TEXT("%-21s UContextualAnimSceneActorComponent::TransitionContextualAnimScene. Can't find AnimSet. Actor: %s SectionName: %s"),
+		UE_LOG(LogContextualAnim, Warning, TEXT("%-21s UContextualAnimSceneActorComponent::TransitionContextualAnimScene. Can't find AnimSet. Actor: %s SectionName: %s"),
 			*UEnum::GetValueAsString(TEXT("Engine.ENetRole"), GetOwner()->GetLocalRole()), *GetNameSafe(GetOwner()), *SectionName.ToString());
 
 		return false;
@@ -533,7 +522,7 @@ bool UContextualAnimSceneActorComponent::TransitionSingleActor(int32 SectionIdx,
 	const int32 AnimSetIdx = Bindings.FindAnimSetForTransitionTo(SectionIdx);
 	if (AnimSetIdx == INDEX_NONE)
 	{
-		UE_LOG(LogContextualAnim, Log, TEXT("%-21s UContextualAnimSceneActorComponent::TransitionSingleActor. Can't find AnimSet. Actor: %s SectionIdx: %d"),
+		UE_LOG(LogContextualAnim, Warning, TEXT("%-21s UContextualAnimSceneActorComponent::TransitionSingleActor. Can't find AnimSet. Actor: %s SectionIdx: %d"),
 			*UEnum::GetValueAsString(TEXT("Engine.ENetRole"), GetOwner()->GetLocalRole()), *GetNameSafe(GetOwner()), SectionIdx);
 
 		return false;
@@ -601,8 +590,9 @@ void UContextualAnimSceneActorComponent::OnTransitionSingleActor(const FContextu
 
 void UContextualAnimSceneActorComponent::OnRep_RepTransitionSingleActor()
 {
-	UE_LOG(LogContextualAnim, Verbose, TEXT("%-21s UContextualAnimSceneActorComponent::OnRep_RepTransitionSingleActor Owner: %s Id: %d RepCounter: %d SectionIdx: %d AnimSetIdx: %d"),
-		*UEnum::GetValueAsString(TEXT("Engine.ENetRole"), GetOwner()->GetLocalRole()), *GetNameSafe(GetOwner()), RepTransitionSingleActorData.Id, RepTransitionSingleActorData.RepCounter, RepTransitionSingleActorData.SectionIdx, RepTransitionSingleActorData.AnimSetIdx);
+	UE_LOG(LogContextualAnim, Log, TEXT("%-21s UContextualAnimSceneActorComponent::OnRep_RepTransitionSingleActor Owner: %s Id: %d RepCounter: %d SectionIdx: %d AnimSetIdx: %d Current Bindings ID: %d"),
+		*UEnum::GetValueAsString(TEXT("Engine.ENetRole"), GetOwner()->GetLocalRole()), *GetNameSafe(GetOwner()), 
+		RepTransitionSingleActorData.Id, RepTransitionSingleActorData.RepCounter, RepTransitionSingleActorData.SectionIdx, RepTransitionSingleActorData.AnimSetIdx, Bindings.IsValid() ? Bindings.GetID() : -1);
 
 	if (!RepTransitionSingleActorData.IsValid())
 	{
@@ -788,7 +778,7 @@ bool UContextualAnimSceneActorComponent::ServerEarlyOutContextualAnimScene_Valid
 
 void UContextualAnimSceneActorComponent::OnRep_TransitionData()
 {
-	UE_LOG(LogContextualAnim, Verbose, TEXT("%-21s UContextualAnimSceneActorComponent::OnRep_TransitionData Actor: %s SectionIdx: %d AnimsetIdx: %d RepCounter: %d"),
+	UE_LOG(LogContextualAnim, Log, TEXT("%-21s UContextualAnimSceneActorComponent::OnRep_TransitionData Actor: %s SectionIdx: %d AnimsetIdx: %d RepCounter: %d"),
 		*UEnum::GetValueAsString(TEXT("Engine.ENetRole"), GetOwner()->GetLocalRole()), *GetNameSafe(GetOwner()), RepTransitionData.SectionIdx, RepTransitionData.AnimSetIdx, RepTransitionData.RepCounter);
 
 	if (!RepTransitionData.IsValid())
@@ -809,7 +799,7 @@ void UContextualAnimSceneActorComponent::OnRep_TransitionData()
 
 void UContextualAnimSceneActorComponent::OnRep_Bindings()
 {
-	UE_LOG(LogContextualAnim, Verbose, TEXT("%-21s UContextualAnimSceneActorComponent::OnRep_Bindings Actor: %s Rep Bindings Id: %d RepCounter: %d Num: %d Current Bindings Id: %d Num: %d"),
+	UE_LOG(LogContextualAnim, Log, TEXT("%-21s UContextualAnimSceneActorComponent::OnRep_Bindings Actor: %s Rep Bindings Id: %d RepCounter: %d Num: %d Current Bindings Id: %d Num: %d"),
 		*UEnum::GetValueAsString(TEXT("Engine.ENetRole"), GetOwner()->GetLocalRole()), *GetNameSafe(GetOwner()), RepBindings.Bindings.GetID(), RepBindings.RepCounter, RepBindings.Bindings.Num(), Bindings.GetID(), Bindings.Num());
 
 	if (!RepBindings.IsValid())
@@ -941,9 +931,11 @@ void UContextualAnimSceneActorComponent::JoinScene(const FContextualAnimSceneBin
 
 	if (const FContextualAnimSceneBinding* Binding = InBindings.FindBindingByActor(GetOwner()))
 	{
-		UE_LOG(LogContextualAnim, Verbose, TEXT("%-21s UContextualAnimSceneActorComponent::JoinScene Actor: %s Role: %s InBindings Id: %d Section: %d Asset: %s"),
+		UE_LOG(LogContextualAnim, Log, TEXT("%-21s UContextualAnimSceneActorComponent::JoinScene Actor: %s Role: %s InBindings Id: %d Section: %d Asset: %s"),
 			*UEnum::GetValueAsString(TEXT("Engine.ENetRole"), GetOwner()->GetLocalRole()), *GetNameSafe(GetOwner()), *InBindings.GetRoleFromBinding(*Binding).ToString(), InBindings.GetID(), InBindings.GetSectionIdx(), *GetNameSafe(InBindings.GetSceneAsset()));
 
+		AnimsPlayed.Reset();
+
 		Bindings = InBindings;
 
 		const FContextualAnimTrack& AnimTrack = Bindings.GetAnimTrackFromBinding(*Binding);
@@ -970,7 +962,7 @@ void UContextualAnimSceneActorComponent::LeaveScene()
 {
 	if (const FContextualAnimSceneBinding* Binding = Bindings.FindBindingByActor(GetOwner()))
 	{
-		UE_LOG(LogContextualAnim, Verbose, TEXT("%-21s UContextualAnimSceneActorComponent::LeaveScene Actor: %s Role: %s Current Bindings Id: %d Section: %d Asset: %s"),
+		UE_LOG(LogContextualAnim, Log, TEXT("%-21s UContextualAnimSceneActorComponent::LeaveScene Actor: %s Role: %s Current Bindings Id: %d Section: %d Asset: %s"),
 			*UEnum::GetValueAsString(TEXT("Engine.ENetRole"), GetOwner()->GetLocalRole()), *GetNameSafe(GetOwner()), *Bindings.GetRoleFromBinding(*Binding).ToString(),
 			Bindings.GetID(), Bindings.GetSectionIdx(), *GetNameSafe(Bindings.GetSceneAsset()));
 
@@ -1003,6 +995,8 @@ void UContextualAnimSceneActorComponent::LeaveScene()
 
 		OnLeftSceneDelegate.Broadcast(this);
 
+		AnimsPlayed.Reset();
+
 		Bindings.Reset();
 	}
 }
@@ -1058,22 +1052,24 @@ bool UContextualAnimSceneActorComponent::CanLeaveScene(const FContextualAnimScen
 
 void UContextualAnimSceneActorComponent::OnMontageBlendingOut(UAnimMontage* Montage, bool bInterrupted)
 {
-	if (bGuardAnimEvents)
-	{
-		return;
-	}
-
-	// Ignore this event if the blending out montage is additive so we don't erroneously leave the interaction when that happens.
-	if (Montage && Montage->IsValidAdditive())
-	{
-		return;
-	}
-
-	UE_LOG(LogContextualAnim, Verbose, TEXT("%-21s UContextualAnimSceneActorComponent::OnMontageBlendingOut Actor: %s Montage: %s bInterrupted: %d"),
+	UE_LOG(LogContextualAnim, Log, TEXT("%-21s UContextualAnimSceneActorComponent::OnMontageBlendingOut Actor: %s Montage: %s bInterrupted: %d"),
 		*UEnum::GetValueAsString(TEXT("Engine.ENetRole"), GetOwner()->GetLocalRole()), *GetNameSafe(GetOwner()), *GetNameSafe(Montage), bInterrupted);
 
 	if (const FContextualAnimSceneBinding* Binding = Bindings.FindBindingByActor(GetOwner()))
 	{
+		AnimsPlayed.RemoveSingleSwap(Montage);
+		if (AnimsPlayed.Num() > 0)
+		{
+			UE_LOG(LogContextualAnim, Log, TEXT("%-21s \tUContextualAnimSceneActorComponent::OnMontageBlendingOut AnimsPlayed Num: %d"), *UEnum::GetValueAsString(TEXT("Engine.ENetRole"), GetOwner()->GetLocalRole()), AnimsPlayed.Num());
+			return;
+		}
+
+		if (!CanLeaveScene(*Binding))
+		{
+			UE_LOG(LogContextualAnim, Log, TEXT("%-21s \tUContextualAnimSceneActorComponent::OnMontageBlendingOut CanLeaveScene FALSE"), *UEnum::GetValueAsString(TEXT("Engine.ENetRole"), GetOwner()->GetLocalRole()));
+			return;
+		}
+
 		const uint8 BindingsId = Bindings.GetID();
 
 		// Stop animation, restore state etc.
@@ -1105,11 +1101,6 @@ void UContextualAnimSceneActorComponent::OnMontageBlendingOut(UAnimMontage* Mont
 
 void UContextualAnimSceneActorComponent::OnPlayMontageNotifyBegin(FName NotifyName, const FBranchingPointNotifyPayload& BranchingPointNotifyPayload)
 {
-	if (bGuardAnimEvents)
-	{
-		return;
-	}
-
 	UE_LOG(LogContextualAnim, Verbose, TEXT("%-21s UContextualAnimSceneActorComponent::OnNotifyBeginReceived Actor: %s Animation: %s NotifyName"),
 		*UEnum::GetValueAsString(TEXT("Engine.ENetRole"), GetOwner()->GetLocalRole()), *GetNameSafe(GetOwner()), *GetNameSafe(BranchingPointNotifyPayload.SequenceAsset), *NotifyName.ToString());
 
diff --git a/Engine/Plugins/Experimental/Animation/ContextualAnimation/Source/ContextualAnimation/Public/ContextualAnimSceneActorComponent.h b/Engine/Plugins/Experimental/Animation/ContextualAnimation/Source/ContextualAnimation/Public/ContextualAnimSceneActorComponent.h
--- a/Engine/Plugins/Experimental/Animation/ContextualAnimation/Source/ContextualAnimation/Public/ContextualAnimSceneActorComponent.h
+++ b/Engine/Plugins/Experimental/Animation/ContextualAnimation/Source/ContextualAnimation/Public/ContextualAnimSceneActorComponent.h
@@ -330,7 +330,8 @@ protected:
 	virtual void OnTransitionSingleActor(const FContextualAnimSceneBinding& Binding, int32 SectionIdx, int32 AnimSetIdx);
 
 private:
-	bool bGuardAnimEvents = false;
+
+	TArray<TWeakObjectPtr<const UAnimMontage>, TInlineAllocator<5>> AnimsPlayed;
 };
 
 #if UE_ENABLE_INCLUDE_ORDER_DEPRECATED_IN_5_2
