Contextual Anim: Changed Collision Behavior default value to None

#virtualized

diff --git a/Engine/Plugins/Experimental/Animation/ContextualAnimation/Source/ContextualAnimation/Public/ContextualAnimSceneAsset.h b/Engine/Plugins/Experimental/Animation/ContextualAnimation/Source/ContextualAnimation/Public/ContextualAnimSceneAsset.h
--- a/Engine/Plugins/Experimental/Animation/ContextualAnimation/Source/ContextualAnimation/Public/ContextualAnimSceneAsset.h
+++ b/Engine/Plugins/Experimental/Animation/ContextualAnimation/Source/ContextualAnimation/Public/ContextualAnimSceneAsset.h
@@ -336,7 +336,7 @@ protected:
 	float Radius = 0.f;
 
 	UPROPERTY(EditAnywhere, Category = "Settings")
-	EContextualAnimCollisionBehavior CollisionBehavior = EContextualAnimCollisionBehavior::IgnoreActorWhenMoving;
+	EContextualAnimCollisionBehavior CollisionBehavior = EContextualAnimCollisionBehavior::None;
 
 	UPROPERTY(EditAnywhere, Category = "Settings", meta = (EditCondition = "CollisionBehavior==EContextualAnimCollisionBehavior::IgnoreChannels", EditConditionHides))
 	TArray<FContextualAnimIgnoreChannelsParam> CollisionChannelsToIgnoreParams;
