Contextual Anim: Added helper function to get motion warping comp from binding and marked a couple of functions as virtual

diff --git a/Engine/Plugins/Experimental/Animation/ContextualAnimation/Source/ContextualAnimation/Private/ContextualAnimTypes.cpp b/Engine/Plugins/Experimental/Animation/ContextualAnimation/Source/ContextualAnimation/Private/ContextualAnimTypes.cpp
--- a/Engine/Plugins/Experimental/Animation/ContextualAnimation/Source/ContextualAnimation/Private/ContextualAnimTypes.cpp
+++ b/Engine/Plugins/Experimental/Animation/ContextualAnimation/Source/ContextualAnimation/Private/ContextualAnimTypes.cpp
@@ -12,6 +12,7 @@
 #include "ContextualAnimSceneAsset.h"
 #include "Components/SkeletalMeshComponent.h"
 #include "GameFramework/CharacterMovementComponent.h"
+#include "MotionWarpingComponent.h"
 
 #include UE_INLINE_GENERATED_CPP_BY_NAME(ContextualAnimTypes)
 
@@ -294,6 +295,19 @@ UCharacterMovementComponent* FContextualAnimSceneBindingContext::GetCharacterMov
 	return CachedMovementComp.Get();
 }
 
+UMotionWarpingComponent* FContextualAnimSceneBindingContext::GetMotionWarpingComponent() const
+{
+	if (!CachedMotionWarpingComp.IsValid() || CachedMotionWarpingComp->GetOwner() != GetActor())
+	{
+		if (Actor.IsValid())
+		{
+			CachedMotionWarpingComp = Actor->FindComponentByClass<UMotionWarpingComponent>();
+		}
+	}
+
+	return CachedMotionWarpingComp.Get();
+}
+
 // FContextualAnimSceneBinding
 ///////////////////////////////////////////////////////////////////////
 
diff --git a/Engine/Plugins/Experimental/Animation/ContextualAnimation/Source/ContextualAnimation/Public/ContextualAnimSceneActorComponent.h b/Engine/Plugins/Experimental/Animation/ContextualAnimation/Source/ContextualAnimation/Public/ContextualAnimSceneActorComponent.h
--- a/Engine/Plugins/Experimental/Animation/ContextualAnimation/Source/ContextualAnimation/Public/ContextualAnimSceneActorComponent.h
+++ b/Engine/Plugins/Experimental/Animation/ContextualAnimation/Source/ContextualAnimation/Public/ContextualAnimSceneActorComponent.h
@@ -266,7 +266,7 @@ protected:
 	 * Used to update IK Targets before animation need them 
 	 */
 	UFUNCTION()
-	void OnTickPose(class USkinnedMeshComponent* SkinnedMeshComponent, float DeltaTime, bool bNeedsValidRootMotion);
+	virtual void OnTickPose(class USkinnedMeshComponent* SkinnedMeshComponent, float DeltaTime, bool bNeedsValidRootMotion);
 
 	UFUNCTION()
 	void OnRep_Bindings();
@@ -319,7 +319,7 @@ protected:
 
 	virtual void PlayAnimation_Internal(UAnimSequenceBase* Animation, float StartTime, bool bSyncPlaybackTime);
 
-	void AddOrUpdateWarpTargets(int32 SectionIdx, int32 AnimSetIdx, const TArray<FContextualAnimWarpPoint>& WarpPoints, const TArray<FContextualAnimWarpTarget>& ExternalWarpTargets);
+	virtual void AddOrUpdateWarpTargets(int32 SectionIdx, int32 AnimSetIdx, const TArray<FContextualAnimWarpPoint>& WarpPoints, const TArray<FContextualAnimWarpTarget>& ExternalWarpTargets);
 
 	void HandleTransitionSelf(int32 NewSectionIdx, int32 NewAnimSetIdx, const TArray<FContextualAnimWarpPoint>& WarpPoints, const TArray<FContextualAnimWarpTarget>& ExternalWarpTargets);
 
diff --git a/Engine/Plugins/Experimental/Animation/ContextualAnimation/Source/ContextualAnimation/Public/ContextualAnimTypes.h b/Engine/Plugins/Experimental/Animation/ContextualAnimation/Source/ContextualAnimation/Public/ContextualAnimTypes.h
--- a/Engine/Plugins/Experimental/Animation/ContextualAnimation/Source/ContextualAnimation/Public/ContextualAnimTypes.h
+++ b/Engine/Plugins/Experimental/Animation/ContextualAnimation/Source/ContextualAnimation/Public/ContextualAnimTypes.h
@@ -18,6 +18,7 @@ class UContextualAnimSelectionCriterion;
 class UContextualAnimSceneAsset;
 class UContextualAnimSceneActorComponent;
 class UCharacterMovementComponent;
+class UMotionWarpingComponent;
 struct FAnimMontageInstance;
 
 namespace UE 
@@ -343,6 +344,8 @@ struct CONTEXTUALANIMATION_API FContextualAnimSceneBindingContext
 
 	UCharacterMovementComponent* GetCharacterMovementComponent() const;
 
+	UMotionWarpingComponent* GetMotionWarpingComponent() const;
+
 	void SetExternalTransform(const FTransform& InTransform);
 
 	FTransform GetTransform() const;
@@ -376,6 +379,9 @@ private:
 	UPROPERTY(NotReplicated)
 	mutable TWeakObjectPtr<UCharacterMovementComponent> CachedMovementComp = nullptr;
 
+	UPROPERTY(NotReplicated)
+	mutable TWeakObjectPtr<UMotionWarpingComponent> CachedMotionWarpingComp = nullptr;
+
 	TOptional<FTransform> ExternalTransform;
 
 	TOptional<FVector> ExternalVelocity;
@@ -401,6 +407,7 @@ struct CONTEXTUALANIMATION_API FContextualAnimSceneBinding
 	FORCEINLINE USkeletalMeshComponent* GetSkeletalMeshComponent() const { return Context.GetSkeletalMeshComponent(); }
 	FORCEINLINE UContextualAnimSceneActorComponent* GetSceneActorComponent() const { return Context.GetSceneActorComponent(); }
 	FORCEINLINE UCharacterMovementComponent* GetCharacterMovementComponent() const { return Context.GetCharacterMovementComponent(); }
+	FORCEINLINE UMotionWarpingComponent* GetMotionWarpingComponent() const { return Context.GetMotionWarpingComponent(); }
 	FORCEINLINE int32 GetAnimTrackIdx() const { return AnimTrackIdx; }
 	
 	void SetAnimTrack(const FContextualAnimTrack& InAnimTrack);
