Contextual Anim: Added option to roles asset to allow user to specify capsule size for preview characters + Added line to the debug text on the viewport showing the current transform relative to the attachment socket for the selected actor (if any)

#virtualized

diff --git a/Engine/Plugins/Experimental/Animation/ContextualAnimation/Source/ContextualAnimation/Public/ContextualAnimTypes.h b/Engine/Plugins/Experimental/Animation/ContextualAnimation/Source/ContextualAnimation/Public/ContextualAnimTypes.h
--- a/Engine/Plugins/Experimental/Animation/ContextualAnimation/Source/ContextualAnimation/Public/ContextualAnimTypes.h
+++ b/Engine/Plugins/Experimental/Animation/ContextualAnimation/Source/ContextualAnimation/Public/ContextualAnimTypes.h
@@ -220,6 +220,12 @@ struct FContextualAnimRoleDefinition
 	UPROPERTY(EditAnywhere, BlueprintReadWrite, Category = "Defaults")
 	bool bIsCharacter = false;
 
+	UPROPERTY(EditAnywhere, BlueprintReadWrite, Category = "Defaults", meta = (EditCondition = "bIsCharacter", EditConditionHides))
+	float PreviewCapsuleHalfHeight = 88.f;
+
+	UPROPERTY(EditAnywhere, BlueprintReadWrite, Category = "Defaults", meta = (EditCondition = "bIsCharacter", EditConditionHides))
+	float PreviewCapsuleRadius = 34.f;
+
 	UPROPERTY(EditAnywhere, BlueprintReadWrite, Category = "Defaults")
 	FTransform MeshToComponent = FTransform(FRotator(0.f, -90.f, 0.f));
 
diff --git a/Engine/Plugins/Experimental/Animation/ContextualAnimation/Source/ContextualAnimationEditor/Private/ContextualAnimViewModel.cpp b/Engine/Plugins/Experimental/Animation/ContextualAnimation/Source/ContextualAnimationEditor/Private/ContextualAnimViewModel.cpp
--- a/Engine/Plugins/Experimental/Animation/ContextualAnimation/Source/ContextualAnimationEditor/Private/ContextualAnimViewModel.cpp
+++ b/Engine/Plugins/Experimental/Animation/ContextualAnimation/Source/ContextualAnimationEditor/Private/ContextualAnimViewModel.cpp
@@ -256,6 +256,8 @@ AActor* FContextualAnimViewModel::SpawnPreviewActor(const FContextualAnimTrack&
 			ACharacter* PreviewCharacter = GetWorld()->SpawnActor<ACharacter>(ACharacter::StaticClass(), SpawnTransform, Params);
 			PreviewCharacter->SetFlags(RF_Transient);
 
+			PreviewCharacter->GetCapsuleComponent()->SetCapsuleSize(RoleDef->PreviewCapsuleRadius, RoleDef->PreviewCapsuleHalfHeight);
+
 			USkeletalMeshComponent* SkelMeshComp = PreviewCharacter->GetMesh();
 			SkelMeshComp->SetRelativeLocation(FVector(0.f, 0.f, -PreviewCharacter->GetCapsuleComponent()->GetScaledCapsuleHalfHeight()));
 			SkelMeshComp->SetRelativeRotation(RoleDef->MeshToComponent.GetRotation());
@@ -1072,8 +1074,25 @@ UContextualAnimSelectionCriterion* FContextualAnimViewModel::GetSelectedSelectio
 FText FContextualAnimViewModel::GetSelectionDebugText() const
 {
 	AActor* SelectedActor = GetSelectedActor();
-	return FText::FromString(FString::Printf(TEXT("Selection Info:\n Role: %s \n Actor: %s \n Criterion: %d (%d)"),
-		*SelectionInfo.Role.ToString(), *GetNameSafe(SelectedActor), SelectionInfo.Criterion.Key, SelectionInfo.Criterion.Value));
+
+	FString AttachmentText = "None";
+	if (SelectedActor)
+	{
+		if (const FContextualAnimAttachmentParams* Params = SceneAsset->GetAttachmentParamsForRole(SelectionInfo.Role))
+		{
+			if (const FContextualAnimSceneBinding* Primary = SceneBindings.GetPrimaryBinding())
+			{
+				if (const UMeshComponent* MeshComp = UContextualAnimUtilities::TryGetMeshComponentWithSocket(Primary->GetActor(), Params->SocketName))
+				{
+					const FTransform SocketTransform = MeshComp->GetSocketTransform(Params->SocketName, ERelativeTransformSpace::RTS_World);
+					const FTransform RelativeTransform = SelectedActor->GetTransform().GetRelativeTransform(SocketTransform);
+					AttachmentText = FString::Printf(TEXT("Location: %s Rotation: %s"), *RelativeTransform.GetLocation().ToString(), *RelativeTransform.Rotator().ToString());
+				}
+			}
+		}
+	}
+	return FText::FromString(FString::Printf(TEXT("Selection Info:\n Role: %s \n Actor: %s \n Attachment Relative Transform: %s \n Criterion: %d (%d)"),
+		*SelectionInfo.Role.ToString(), *GetNameSafe(SelectedActor), *AttachmentText, SelectionInfo.Criterion.Key, SelectionInfo.Criterion.Value));
 }
 
 bool FContextualAnimViewModel::ProcessInputDelta(FVector& InDrag, FRotator& InRot, FVector& InScale)
