Contextual Anim: Added safety check to prevent a crash that could happen in a rare case where a transition is attempted on an actor that is not part of an interaction anymore.

diff --git a/Engine/Plugins/Experimental/Animation/ContextualAnimation/Source/ContextualAnimation/Private/ContextualAnimSceneActorComponent.cpp b/Engine/Plugins/Experimental/Animation/ContextualAnimation/Source/ContextualAnimation/Private/ContextualAnimSceneActorComponent.cpp
--- a/Engine/Plugins/Experimental/Animation/ContextualAnimation/Source/ContextualAnimation/Private/ContextualAnimSceneActorComponent.cpp
+++ b/Engine/Plugins/Experimental/Animation/ContextualAnimation/Source/ContextualAnimation/Private/ContextualAnimSceneActorComponent.cpp
@@ -485,6 +485,13 @@ void UContextualAnimSceneActorComponent::HandleTransitionEveryone(int32 NewSecti
 
 void UContextualAnimSceneActorComponent::HandleTransitionSelf(int32 NewSectionIdx, int32 NewAnimSetIdx, const TArray<FContextualAnimWarpPoint>& WarpPoints, const TArray<FContextualAnimWarpTarget>& ExternalWarpTargets)
 {
+	if (!Bindings.IsValid())
+	{
+		UE_LOG(LogContextualAnim, Warning, TEXT("%-21s UContextualAnimSceneActorComponent::HandleTransitionSelf Invalid Bindings. Actor: %s NewSectionIdx: %d NewAnimSetIdx: %d"),
+			*UEnum::GetValueAsString(TEXT("Engine.ENetRole"), GetOwner()->GetLocalRole()), *GetNameSafe(GetOwner()), NewSectionIdx, NewAnimSetIdx);
+		return;
+	}
+	
 	// Update bindings internal data so it points to the new section and new anim set
 	Bindings.TransitionTo(NewSectionIdx, NewAnimSetIdx);
 
