Contextual Anim: When an actor leave an interaction it now notifies the other actors so they can properly restore collision and update bindings.

diff --git a/Engine/Plugins/Experimental/Animation/ContextualAnimation/Source/ContextualAnimation/Private/ContextualAnimSceneActorComponent.cpp b/Engine/Plugins/Experimental/Animation/ContextualAnimation/Source/ContextualAnimation/Private/ContextualAnimSceneActorComponent.cpp
--- a/Engine/Plugins/Experimental/Animation/ContextualAnimation/Source/ContextualAnimation/Private/ContextualAnimSceneActorComponent.cpp
+++ b/Engine/Plugins/Experimental/Animation/ContextualAnimation/Source/ContextualAnimation/Private/ContextualAnimSceneActorComponent.cpp
@@ -857,14 +857,19 @@ void UContextualAnimSceneActorComponent::SetIgnoreCollisionWithOtherActors(bool
 		AActor* OtherActor = Binding.GetActor();
 		if (OtherActor && OtherActor != OwnerActor)
 		{
-			if (UPrimitiveComponent* RootPrimitiveComponent = Cast<UPrimitiveComponent>(OwnerActor->GetRootComponent()))
-			{
-				RootPrimitiveComponent->IgnoreActorWhenMoving(OtherActor, bValue);
-			}
+			SetIgnoreCollisionWithActor(*OtherActor, bValue);
 		}
 	}
 }
 
+void UContextualAnimSceneActorComponent::SetIgnoreCollisionWithActor(AActor& Actor, bool bValue) const
+{
+	if (UPrimitiveComponent* RootPrimitiveComponent = Cast<UPrimitiveComponent>(GetOwner()->GetRootComponent()))
+	{
+		RootPrimitiveComponent->IgnoreActorWhenMoving(&Actor, bValue);
+	}
+}
+
 void UContextualAnimSceneActorComponent::SetCollisionState(const FContextualAnimSceneBinding& Binding)
 {
 	if (const UContextualAnimSceneAsset* Asset = Bindings.GetSceneAsset())
@@ -998,6 +1003,20 @@ void UContextualAnimSceneActorComponent::LeaveScene()
 
 		RestoreMovementState(*Binding);
 
+		// Notify the other actors in the interaction
+		// @TODO: This should be refactored so only the leader of the interaction maintains the full bindings
+		for (const FContextualAnimSceneBinding& OtherBinding : Bindings)
+		{
+			AActor* OwnerActor = GetOwner();
+			if (OtherBinding.GetActor() != OwnerActor)
+			{
+				if (UContextualAnimSceneActorComponent* Comp = OtherBinding.GetSceneActorComponent())
+				{
+					Comp->OtherActorLeftScene(*OwnerActor);
+				}
+			}
+		}
+
 		OnLeaveScene(*Binding);
 
 		OnLeftSceneDelegate.Broadcast(this);
@@ -1008,6 +1027,23 @@ void UContextualAnimSceneActorComponent::LeaveScene()
 	}
 }
 
+void UContextualAnimSceneActorComponent::OtherActorLeftScene(AActor& Actor)
+{
+	if (Bindings.IsValid())
+	{
+		if (const UContextualAnimSceneAsset* Asset = Bindings.GetSceneAsset())
+		{
+			const EContextualAnimCollisionBehavior CollisionBehavior = Asset->GetCollisionBehavior();
+			if (CollisionBehavior == EContextualAnimCollisionBehavior::IgnoreActorWhenMoving)
+			{
+				SetIgnoreCollisionWithActor(Actor, false);
+			}
+
+			Bindings.RemoveActor(Actor);
+		}
+	}
+}
+
 void UContextualAnimSceneActorComponent::OnLeaveScene(const FContextualAnimSceneBinding& Binding)
 {
 	// For derived classes to override.
@@ -1105,7 +1141,6 @@ void UContextualAnimSceneActorComponent::OnMontageBlendingOut(UAnimMontage* Mont
 	}
 }
 
-
 void UContextualAnimSceneActorComponent::OnPlayMontageNotifyBegin(FName NotifyName, const FBranchingPointNotifyPayload& BranchingPointNotifyPayload)
 {
 	UE_LOG(LogContextualAnim, Verbose, TEXT("%-21s UContextualAnimSceneActorComponent::OnNotifyBeginReceived Actor: %s Animation: %s NotifyName"),
diff --git a/Engine/Plugins/Experimental/Animation/ContextualAnimation/Source/ContextualAnimation/Private/ContextualAnimTypes.cpp b/Engine/Plugins/Experimental/Animation/ContextualAnimation/Source/ContextualAnimation/Private/ContextualAnimTypes.cpp
--- a/Engine/Plugins/Experimental/Animation/ContextualAnimation/Source/ContextualAnimation/Private/ContextualAnimTypes.cpp
+++ b/Engine/Plugins/Experimental/Animation/ContextualAnimation/Source/ContextualAnimation/Private/ContextualAnimTypes.cpp
@@ -889,3 +889,17 @@ void FContextualAnimSceneBindings::TransitionTo(int32 NewSectionIdx, int32 NewAn
 		}
 	}
 }
+
+bool FContextualAnimSceneBindings::RemoveActor(AActor& ActorRef)
+{
+	for (int32 Idx = 0; Idx < Data.Num(); Idx++)
+	{
+		if (Data[Idx].GetActor() == &ActorRef)
+		{
+			Data.RemoveAt(Idx);
+			return true;
+		}
+	}
+
+	return false;
+}
\ No newline at end of file
diff --git a/Engine/Plugins/Experimental/Animation/ContextualAnimation/Source/ContextualAnimation/Public/ContextualAnimSceneActorComponent.h b/Engine/Plugins/Experimental/Animation/ContextualAnimation/Source/ContextualAnimation/Public/ContextualAnimSceneActorComponent.h
--- a/Engine/Plugins/Experimental/Animation/ContextualAnimation/Source/ContextualAnimation/Public/ContextualAnimSceneActorComponent.h
+++ b/Engine/Plugins/Experimental/Animation/ContextualAnimation/Source/ContextualAnimation/Public/ContextualAnimSceneActorComponent.h
@@ -280,6 +280,7 @@ protected:
 	UFUNCTION()
 	void OnRep_TransitionData();
 
+	void SetIgnoreCollisionWithActor(AActor& Actor, bool bValue) const;
 	void SetIgnoreCollisionWithOtherActors(bool bValue) const;
 	
 	virtual void SetCollisionState(const FContextualAnimSceneBinding& Binding);
@@ -329,6 +330,8 @@ protected:
 
 	virtual void OnTransitionSingleActor(const FContextualAnimSceneBinding& Binding, int32 SectionIdx, int32 AnimSetIdx);
 
+	void OtherActorLeftScene(AActor& Actor);
+
 private:
 
 	TArray<TWeakObjectPtr<const UAnimMontage>, TInlineAllocator<5>> AnimsPlayed;
diff --git a/Engine/Plugins/Experimental/Animation/ContextualAnimation/Source/ContextualAnimation/Public/ContextualAnimTypes.h b/Engine/Plugins/Experimental/Animation/ContextualAnimation/Source/ContextualAnimation/Public/ContextualAnimTypes.h
--- a/Engine/Plugins/Experimental/Animation/ContextualAnimation/Source/ContextualAnimation/Public/ContextualAnimTypes.h
+++ b/Engine/Plugins/Experimental/Animation/ContextualAnimation/Source/ContextualAnimation/Public/ContextualAnimTypes.h
@@ -509,6 +509,8 @@ struct CONTEXTUALANIMATION_API FContextualAnimSceneBindings
 	int32 FindAnimSetForTransitionTo(int32 NewSectionIdx) const;
 	void TransitionTo(int32 NewSectionIdx, int32 NewAnimSetIdx);
 
+	bool RemoveActor(AActor& ActorRef);
+
 private:
 
 	friend class UContextualAnimManager;
