Contextual Anim: Added WITH_EDITORONLY_DATA macro around preview data member

diff --git a/Engine/Plugins/Experimental/Animation/ContextualAnimation/Source/ContextualAnimation/Public/ContextualAnimSceneAsset.h b/Engine/Plugins/Experimental/Animation/ContextualAnimation/Source/ContextualAnimation/Public/ContextualAnimSceneAsset.h
--- a/Engine/Plugins/Experimental/Animation/ContextualAnimation/Source/ContextualAnimation/Public/ContextualAnimSceneAsset.h
+++ b/Engine/Plugins/Experimental/Animation/ContextualAnimation/Source/ContextualAnimation/Public/ContextualAnimSceneAsset.h
@@ -326,9 +326,13 @@ protected:
 	UPROPERTY(EditAnywhere, Category = "Settings", meta = (GetOptions = "GetRoles"))
 	FName PrimaryRole = NAME_None;
 
+#if WITH_EDITORONLY_DATA
+
 	UPROPERTY(EditAnywhere, Category = "Settings", meta = (TitleProperty = "Role"))
 	TArray<FContextualAnimActorPreviewData> OverridePreviewData;
 
+#endif // WITH_EDITORONLY_DATA
+
 	UPROPERTY(EditAnywhere, Category = "Defaults")
 	TArray<FContextualAnimSceneSection> Sections;
 
