Contextual Anim: Fix for crash when switching sections after removing role

diff --git a/Engine/Plugins/Experimental/Animation/ContextualAnimation/Source/ContextualAnimationEditor/Private/ContextualAnimMovieSceneTrack.cpp b/Engine/Plugins/Experimental/Animation/ContextualAnimation/Source/ContextualAnimationEditor/Private/ContextualAnimMovieSceneTrack.cpp
--- a/Engine/Plugins/Experimental/Animation/ContextualAnimation/Source/ContextualAnimationEditor/Private/ContextualAnimMovieSceneTrack.cpp
+++ b/Engine/Plugins/Experimental/Animation/ContextualAnimation/Source/ContextualAnimationEditor/Private/ContextualAnimMovieSceneTrack.cpp
@@ -5,9 +5,13 @@
 #include "ContextualAnimMovieSceneSequence.h"
 #include "MovieScene.h"
 #include "ContextualAnimTypes.h"
+#include "ContextualAnimSceneAsset.h"
+#include "ContextualAnimViewModel.h"
 
 #include UE_INLINE_GENERATED_CPP_BY_NAME(ContextualAnimMovieSceneTrack)
 
+#define LOCTEXT_NAMESPACE "ContexualAnimTrack"
+
 UContextualAnimMovieSceneTrack::UContextualAnimMovieSceneTrack()
 {
 #if WITH_EDITORONLY_DATA
@@ -87,4 +91,43 @@ EMovieSceneSectionMovedResult UContextualAnimMovieSceneTrack::OnSectionMoved(UMo
 
 	return EMovieSceneSectionMovedResult::None;
 }
+
+FSlateColor UContextualAnimMovieSceneTrack::GetLabelColor(const FMovieSceneLabelParams& LabelParams) const
+{
+	if (Role != NAME_None)
+	{
+		if (const UContextualAnimSceneAsset* SceneAsset = GetViewModel().GetSceneAsset())
+		{
+			const TArray<FName> ValidRoles = SceneAsset->GetRoles();
+			if (ValidRoles.Contains(Role) == false)
+			{
+				return FSlateColor(FColor::Red);
+			}
+		}
+	}
+
+	return Super::GetLabelColor(LabelParams);
+}
+
+FText UContextualAnimMovieSceneTrack::GetDisplayNameToolTipText(const FMovieSceneLabelParams& LabelParams) const
+{
+	if (Role != NAME_None)
+	{
+		if (const UContextualAnimSceneAsset* SceneAsset = GetViewModel().GetSceneAsset())
+		{
+			const TArray<FName> ValidRoles = SceneAsset->GetRoles();
+			if (ValidRoles.Contains(Role) == false)
+			{
+				return FText::Format(LOCTEXT("ContextualAnimMovieSceneTrack_MissingRole", "Role {0} is missing from Role Asset {1}. Add the role to the Roles Asset, or Update Roles in the details panel."), 
+					FText::FromName(Role), FText::FromString(GetNameSafe(SceneAsset)));
+			}
+		}
+	}
+
+
+	return Super::GetDisplayNameToolTipText(LabelParams);
+}
+
 #endif
+
+#undef LOCTEXT_NAMESPACE
\ No newline at end of file
diff --git a/Engine/Plugins/Experimental/Animation/ContextualAnimation/Source/ContextualAnimationEditor/Private/ContextualAnimViewModel.cpp b/Engine/Plugins/Experimental/Animation/ContextualAnimation/Source/ContextualAnimationEditor/Private/ContextualAnimViewModel.cpp
--- a/Engine/Plugins/Experimental/Animation/ContextualAnimation/Source/ContextualAnimationEditor/Private/ContextualAnimViewModel.cpp
+++ b/Engine/Plugins/Experimental/Animation/ContextualAnimation/Source/ContextualAnimationEditor/Private/ContextualAnimViewModel.cpp
@@ -389,8 +389,23 @@ void FContextualAnimViewModel::SetDefaultMode()
 		return;
 	}
 
-	TArray<FName> Roles = SceneAsset->GetRoles();
-	for (FName Role : Roles)
+	FContextualAnimSceneSection& ContextualAnimSection = SceneAsset->Sections[ActiveSectionIdx];
+
+	// If the roles asset was edited, our contextual anim asset may have tracks for roles that don't exist in the roles asset.
+	// We allow this, log a warning, and display this in the UI for the user.
+	TArray<FName> ValidRoles = SceneAsset->GetRoles();
+	TArray<FName> ActualRoles;
+	for (int32 AnimSetIdx = 0; AnimSetIdx < ContextualAnimSection.AnimSets.Num(); AnimSetIdx++)
+	{
+		FContextualAnimSet& AnimSet = ContextualAnimSection.AnimSets[AnimSetIdx];
+		for (int32 AnimTrackIdx = 0; AnimTrackIdx < AnimSet.Tracks.Num(); AnimTrackIdx++)
+		{
+			const FContextualAnimTrack& AnimTrack = AnimSet.Tracks[AnimTrackIdx];
+			ActualRoles.AddUnique(AnimTrack.Role);
+		}
+	}
+
+	for (FName Role : ActualRoles)
 	{
 		UContextualAnimMovieSceneTrack* MovieSceneAnimTrack = MovieSceneSequence->GetMovieScene()->AddTrack<UContextualAnimMovieSceneTrack>();
 		check(MovieSceneAnimTrack);
@@ -398,8 +413,6 @@ void FContextualAnimViewModel::SetDefaultMode()
 		MovieSceneAnimTrack->Initialize(Role);
 	}
 
-	FContextualAnimSceneSection& ContextualAnimSection = SceneAsset->Sections[ActiveSectionIdx];
-
 	float EmptyAnimSectionLength = 0.f;
 	for (int32 AnimSetIdx = 0; AnimSetIdx < ContextualAnimSection.AnimSets.Num(); AnimSetIdx++)
 	{
@@ -457,6 +470,13 @@ void FContextualAnimViewModel::SetDefaultMode()
 
 			MovieSceneTrack->AddSection(*NewSection);
 			MovieSceneTrack->SetTrackRowDisplayName(FText::FromString(FString::Printf(TEXT("%d"), AnimSetIdx)), AnimSetIdx);
+
+			const bool bIsValidRole = ValidRoles.Contains(AnimTrack.Role);
+			if (bIsValidRole == false)
+			{
+				UE_LOG(LogContextualAnim, Warning, TEXT("Scene Asset: %s - Role Asset %s is missing Role %s. Please use Update Roles under settings, or update your roles asset to include this track."), 
+												*GetNameSafe(SceneAsset), *GetNameSafe(SceneAsset->GetRolesAsset()), *AnimTrack.Role.ToString());
+			}
 		}
 	}
 
@@ -1313,6 +1333,31 @@ void FContextualAnimViewModel::DiscardChangeToActorTransformInScene()
 	ModifyingTransformInSceneCachedActor.Reset();
 }
 
+void FContextualAnimViewModel::UpdateRoles()
+{
+	FScopedTransaction Transaction(LOCTEXT("CAS_UpdateRoles", "UpdateRoles"));
+	SceneAsset->Modify();
+
+	TArray<FName> ValidRoles;
+	if (SceneAsset)
+	{
+		ValidRoles = SceneAsset->GetRoles();
+	}
+
+	for(FContextualAnimSceneSection& Section : SceneAsset->Sections)
+	{
+		for (FContextualAnimSet& AnimSet : Section.AnimSets)
+		{
+			AnimSet.Tracks.RemoveAll([&ValidRoles](const FContextualAnimTrack& Track)->bool
+				{
+					return ValidRoles.Contains(Track.Role) == false;
+				});
+		}
+	}
+
+	SetDefaultMode();
+}
+
 void FContextualAnimViewModel::CacheWarpPoints()
 {
 	for(FContextualAnimSceneSection& Section : SceneAsset->Sections)
diff --git a/Engine/Plugins/Experimental/Animation/ContextualAnimation/Source/ContextualAnimationEditor/Private/DetailCustomizations/ContextualAnimSceneAssetDetailCustom.cpp b/Engine/Plugins/Experimental/Animation/ContextualAnimation/Source/ContextualAnimationEditor/Private/DetailCustomizations/ContextualAnimSceneAssetDetailCustom.cpp
--- a/Engine/Plugins/Experimental/Animation/ContextualAnimation/Source/ContextualAnimationEditor/Private/DetailCustomizations/ContextualAnimSceneAssetDetailCustom.cpp
+++ b/Engine/Plugins/Experimental/Animation/ContextualAnimation/Source/ContextualAnimationEditor/Private/DetailCustomizations/ContextualAnimSceneAssetDetailCustom.cpp
@@ -36,6 +36,25 @@ void FContextualAnimSceneAssetDetailCustom::CustomizeDetails(IDetailLayoutBuilde
 	IDetailCategoryBuilder& SettingsCategory = DetailBuilder.EditCategory(TEXT("Settings"));
 	SettingsCategory.SetSortOrder(0);
 
+	// Add button to update role tracks for the scene
+	SettingsCategory.AddCustomRow(FText::GetEmpty())
+		.ValueContent()
+		.VAlign(VAlign_Center)
+		.MaxDesiredWidth(250)
+		[
+			SNew(SButton).Text(LOCTEXT("UpdateRolesPointsLabel", "Update Roles"))
+			.OnClicked_Lambda([ViewModel]()
+				{
+					const FText DialogMsg = LOCTEXT("UpdateRolesPointsDialog", "Updating the roles will remove any tracks for roles that don't exist in your Roles Asset. Are you sure you want to continue?");
+					if (FMessageDialog::Open(EAppMsgType::YesNo, DialogMsg) == EAppReturnType::Yes)
+					{
+						ViewModel->UpdateRoles();
+					}
+
+					return FReply::Handled();
+				})
+		];
+
 	IDetailCategoryBuilder& Category = DetailBuilder.EditCategory(TEXT("Defaults"));
 	Category.SetCategoryVisibility(false);
 
diff --git a/Engine/Plugins/Experimental/Animation/ContextualAnimation/Source/ContextualAnimationEditor/Public/ContextualAnimMovieSceneTrack.h b/Engine/Plugins/Experimental/Animation/ContextualAnimation/Source/ContextualAnimationEditor/Public/ContextualAnimMovieSceneTrack.h
--- a/Engine/Plugins/Experimental/Animation/ContextualAnimation/Source/ContextualAnimationEditor/Public/ContextualAnimMovieSceneTrack.h
+++ b/Engine/Plugins/Experimental/Animation/ContextualAnimation/Source/ContextualAnimationEditor/Public/ContextualAnimMovieSceneTrack.h
@@ -52,6 +52,10 @@ public:
 #if WITH_EDITOR
 	/** Called when a section in this track is moved. Updates the actual Notify the section is representing */
 	virtual EMovieSceneSectionMovedResult OnSectionMoved(UMovieSceneSection& Section, const FMovieSceneSectionMovedParams& Params) override;
+
+	virtual FSlateColor GetLabelColor(const FMovieSceneLabelParams& LabelParams) const override;
+	virtual FText GetDisplayNameToolTipText(const FMovieSceneLabelParams& LabelParams) const override;
+
 #endif
 
 	virtual bool SupportsMultipleRows() const override { return true; }
@@ -65,7 +69,6 @@ private:
 	/** Role this track represents */
 	UPROPERTY()
 	FName Role = NAME_None;
-
 };
 
 
diff --git a/Engine/Plugins/Experimental/Animation/ContextualAnimation/Source/ContextualAnimationEditor/Public/ContextualAnimViewModel.h b/Engine/Plugins/Experimental/Animation/ContextualAnimation/Source/ContextualAnimationEditor/Public/ContextualAnimViewModel.h
--- a/Engine/Plugins/Experimental/Animation/ContextualAnimation/Source/ContextualAnimationEditor/Public/ContextualAnimViewModel.h
+++ b/Engine/Plugins/Experimental/Animation/ContextualAnimation/Source/ContextualAnimationEditor/Public/ContextualAnimViewModel.h
@@ -131,6 +131,7 @@ public:
 	bool CanMakeEdits() const;
 
 	void CacheWarpPoints();
+	void UpdateRoles();
 
 private:
 
