Contextual Anim Editor - Fix for crash happening when multiple cas editors are open and one of them is closed

diff --git a/Engine/Plugins/Experimental/Animation/ContextualAnimation/Source/ContextualAnimationEditor/Private/ContextualAnimMovieSceneNotifyTrackEditor.cpp b/Engine/Plugins/Experimental/Animation/ContextualAnimation/Source/ContextualAnimationEditor/Private/ContextualAnimMovieSceneNotifyTrackEditor.cpp
--- a/Engine/Plugins/Experimental/Animation/ContextualAnimation/Source/ContextualAnimationEditor/Private/ContextualAnimMovieSceneNotifyTrackEditor.cpp
+++ b/Engine/Plugins/Experimental/Animation/ContextualAnimation/Source/ContextualAnimationEditor/Private/ContextualAnimMovieSceneNotifyTrackEditor.cpp
@@ -89,7 +89,7 @@ void FContextualAnimMovieSceneNotifyTrackEditor::CustomizeToolBar(FToolBarBuilde
 		}));
 
 		Action.IsActionVisibleDelegate = FIsActionButtonVisible::CreateLambda([this]() {
-			return GetMovieSceneSequence().GetViewModel().GetTimelineMode() == FContextualAnimViewModel::ETimelineMode::Notifies;
+			return (GetSequencer() != nullptr) ? (GetMovieSceneSequence().GetViewModel().GetTimelineMode() == FContextualAnimViewModel::ETimelineMode::Notifies) : false;
 		});
 
 		ToolBarBuilder.AddToolBarButton(
