Contextual Anim Editor
- Added missing parentisis to fix order of operations in the function that update the preview actor position in editor when scrubbing the timeline fixing a bug that was causing the actors to be sliglt offset from where they should be relative to each other.
- Updated logic around MeshToComponent Transfrom in the roles asset so we can use it to specify a location offset too to be able to get a representation of the character in the preview scene identical to what we would have in game.

#changelist validated
#virtualized

diff --git a/Engine/Plugins/Experimental/Animation/ContextualAnimation/Source/ContextualAnimation/Private/ContextualAnimSceneAsset.cpp b/Engine/Plugins/Experimental/Animation/ContextualAnimation/Source/ContextualAnimation/Private/ContextualAnimSceneAsset.cpp
--- a/Engine/Plugins/Experimental/Animation/ContextualAnimation/Source/ContextualAnimation/Private/ContextualAnimSceneAsset.cpp
+++ b/Engine/Plugins/Experimental/Animation/ContextualAnimation/Source/ContextualAnimation/Private/ContextualAnimSceneAsset.cpp
@@ -102,7 +102,7 @@ void FContextualAnimSceneSection::GenerateAlignmentTracks(UContextualAnimSceneAs
 		{
 			UE_LOG(LogContextualAnim, Log, TEXT("Generating AlignmentTracks Tracks. Animation: %s"), *GetNameSafe(AnimTrack.Animation));
 
-			const FTransform MeshToComponentInverse = SceneAsset.GetMeshToComponentForRole(AnimTrack.Role).Inverse();
+			const FTransform MeshToComponentInverse = FTransform(SceneAsset.GetMeshToComponentForRole(AnimTrack.Role).GetRotation()).Inverse();
 			const float SampleInterval = 1.f / SceneAsset.GetSampleRate();
 
 			// Initialize tracks for each alignment section
@@ -715,7 +715,7 @@ FTransform UContextualAnimSceneAsset::GetAlignmentTransform(const FContextualAni
 		const FName WarpTargetName = Section->GetWarpPointDefinitions()[WarpPointIdx].WarpTargetName;
 		if (const FTransform* WarpPointTransformPtr = AnimSet->WarpPoints.Find(WarpTargetName))
 		{
-			const FTransform MeshToComponentInverse = GetMeshToComponentForRole(AnimTrack.Role).Inverse();
+			const FTransform MeshToComponentInverse = FTransform(GetMeshToComponentForRole(AnimTrack.Role).GetRotation()).Inverse();
 
 			if (AnimTrack.Animation)
 			{
diff --git a/Engine/Plugins/Experimental/Animation/ContextualAnimation/Source/ContextualAnimation/Private/ContextualAnimUtilities.cpp b/Engine/Plugins/Experimental/Animation/ContextualAnimation/Source/ContextualAnimation/Private/ContextualAnimUtilities.cpp
--- a/Engine/Plugins/Experimental/Animation/ContextualAnimation/Source/ContextualAnimation/Private/ContextualAnimUtilities.cpp
+++ b/Engine/Plugins/Experimental/Animation/ContextualAnimation/Source/ContextualAnimation/Private/ContextualAnimUtilities.cpp
@@ -166,7 +166,7 @@ void UContextualAnimUtilities::DrawDebugAnimSet(const UWorld* World, const UCont
 	{
 		for (const FContextualAnimTrack& AnimTrack : AnimSet.Tracks)
 		{
-			const FTransform Transform = (SceneAsset.GetMeshToComponentForRole(AnimTrack.Role) * SceneAsset.GetAlignmentTransform(AnimTrack, 0, Time)) * ToWorldTransform;
+			const FTransform Transform = (FTransform(SceneAsset.GetMeshToComponentForRole(AnimTrack.Role).GetRotation()) * SceneAsset.GetAlignmentTransform(AnimTrack, 0, Time)) * ToWorldTransform;
 
 			if (const UAnimSequenceBase* Animation = AnimTrack.Animation)
 			{
diff --git a/Engine/Plugins/Experimental/Animation/ContextualAnimation/Source/ContextualAnimation/Public/ContextualAnimTypes.h b/Engine/Plugins/Experimental/Animation/ContextualAnimation/Source/ContextualAnimation/Public/ContextualAnimTypes.h
--- a/Engine/Plugins/Experimental/Animation/ContextualAnimation/Source/ContextualAnimation/Public/ContextualAnimTypes.h
+++ b/Engine/Plugins/Experimental/Animation/ContextualAnimation/Source/ContextualAnimation/Public/ContextualAnimTypes.h
@@ -227,7 +227,7 @@ struct FContextualAnimRoleDefinition
 	float PreviewCapsuleRadius = 34.f;
 
 	UPROPERTY(EditAnywhere, BlueprintReadWrite, Category = "Defaults")
-	FTransform MeshToComponent = FTransform(FRotator(0.f, -90.f, 0.f));
+	FTransform MeshToComponent = FTransform(FRotator(0.f, 0.f, 0.f));
 
 	static const FContextualAnimRoleDefinition InvalidRoleDefinition;
 };
diff --git a/Engine/Plugins/Experimental/Animation/ContextualAnimation/Source/ContextualAnimationEditor/Private/ContextualAnimEdMode.cpp b/Engine/Plugins/Experimental/Animation/ContextualAnimation/Source/ContextualAnimationEditor/Private/ContextualAnimEdMode.cpp
--- a/Engine/Plugins/Experimental/Animation/ContextualAnimation/Source/ContextualAnimationEditor/Private/ContextualAnimEdMode.cpp
+++ b/Engine/Plugins/Experimental/Animation/ContextualAnimation/Source/ContextualAnimationEditor/Private/ContextualAnimEdMode.cpp
@@ -199,7 +199,7 @@ void FContextualAnimEdMode::Render(const FSceneView* View, FViewport* Viewport,
 								{
 									const float Time = 0.f;
 									const FTransform RelativeToPrimary = SceneAsset->GetAlignmentTransformForRoleRelativeToOtherRole(AnimTrack.SectionIdx, AnimTrack.AnimSetIdx, AnimTrack.Role, PrimaryRole, Time);
-									FTransform Transform = (SceneAsset->GetMeshToComponentForRole(AnimTrack.Role) * RelativeToPrimary) * PrimaryTransform;
+									FTransform Transform = (FTransform(SceneAsset->GetMeshToComponentForRole(AnimTrack.Role).GetRotation()) * RelativeToPrimary) * PrimaryTransform;
 
 									if (ACharacter* PrimaryAsCharacter = Cast<ACharacter>(PrimaryBinding->GetActor()))
 									{
diff --git a/Engine/Plugins/Experimental/Animation/ContextualAnimation/Source/ContextualAnimationEditor/Private/ContextualAnimViewModel.cpp b/Engine/Plugins/Experimental/Animation/ContextualAnimation/Source/ContextualAnimationEditor/Private/ContextualAnimViewModel.cpp
--- a/Engine/Plugins/Experimental/Animation/ContextualAnimation/Source/ContextualAnimationEditor/Private/ContextualAnimViewModel.cpp
+++ b/Engine/Plugins/Experimental/Animation/ContextualAnimation/Source/ContextualAnimationEditor/Private/ContextualAnimViewModel.cpp
@@ -259,8 +259,7 @@ AActor* FContextualAnimViewModel::SpawnPreviewActor(const FContextualAnimTrack&
 			PreviewCharacter->GetCapsuleComponent()->SetCapsuleSize(RoleDef->PreviewCapsuleRadius, RoleDef->PreviewCapsuleHalfHeight);
 
 			USkeletalMeshComponent* SkelMeshComp = PreviewCharacter->GetMesh();
-			SkelMeshComp->SetRelativeLocation(FVector(0.f, 0.f, -PreviewCharacter->GetCapsuleComponent()->GetScaledCapsuleHalfHeight()));
-			SkelMeshComp->SetRelativeRotation(RoleDef->MeshToComponent.GetRotation());
+			SkelMeshComp->SetRelativeTransform(RoleDef->MeshToComponent);
 			SkelMeshComp->SetSkeletalMesh(PreviewSkeletalMesh);
 			SkelMeshComp->SetAnimationMode(EAnimationMode::AnimationBlueprint);
 
@@ -843,9 +842,9 @@ void FContextualAnimViewModel::UpdatePreviewActorTransform(const FContextualAnim
 				MovementComp->StopMovementImmediately();
 			}
 
-			const float MIN_FLOOR_DIST = 1.9f; //from CharacterMovementComp, including in this offset to avoid jittering in walking mode
+			const float MIN_FLOOR_DIST = 1.9f; //from CharacterMovementComp, including this offset to avoid jittering in walking mode
 			const float CapsuleHalfHeight = PreviewCharacter->GetCapsuleComponent()->GetScaledCapsuleHalfHeight();
-			Transform.SetLocation(Transform.GetLocation() + (PreviewCharacter->GetActorQuat().GetUpVector() * CapsuleHalfHeight + MIN_FLOOR_DIST));
+			Transform.SetLocation(Transform.GetLocation() + (PreviewCharacter->GetActorQuat().GetUpVector() * (CapsuleHalfHeight + MIN_FLOOR_DIST)));
 
 			Transform.SetRotation(PreviewCharacter->GetBaseRotationOffset().Inverse() * Transform.GetRotation());
 		}
@@ -1375,7 +1374,7 @@ void FContextualAnimViewModel::CacheWarpPoints()
 					if (AnimTrack)
 					{
 						const FTransform RootTransform = AnimTrack->Animation ? UContextualAnimUtilities::ExtractRootTransformFromAnimation(AnimTrack->Animation, 0.f) : FTransform::Identity;
-						const FTransform WarpPointTransform = (SceneAsset->GetMeshToComponentForRole(AnimTrack->Role).Inverse() * (RootTransform * AnimTrack->MeshToScene));
+						const FTransform WarpPointTransform = (FTransform(SceneAsset->GetMeshToComponentForRole(AnimTrack->Role).GetRotation()).Inverse() * (RootTransform * AnimTrack->MeshToScene));
 						AnimSet.WarpPoints.Add(WarpPointDef.WarpTargetName, WarpPointTransform);
 					}
 				}
@@ -1404,10 +1403,10 @@ void FContextualAnimViewModel::CacheWarpPoints()
 							if (OtherAnimTrack)
 							{
 								FTransform T1 = AnimTrack->Animation ? UContextualAnimUtilities::ExtractRootTransformFromAnimation(AnimTrack->Animation, 0.f) : FTransform::Identity;
-								T1 = (SceneAsset->GetMeshToComponentForRole(AnimTrack->Role).Inverse() * (T1 * AnimTrack->MeshToScene));
+								T1 = (FTransform(SceneAsset->GetMeshToComponentForRole(AnimTrack->Role).GetRotation()).Inverse() * (T1 * AnimTrack->MeshToScene));
 
 								FTransform T2 = OtherAnimTrack->Animation ? UContextualAnimUtilities::ExtractRootTransformFromAnimation(OtherAnimTrack->Animation, 0.f) : FTransform::Identity;
-								T2 = (SceneAsset->GetMeshToComponentForRole(OtherAnimTrack->Role).Inverse() * (T2 * OtherAnimTrack->MeshToScene));
+								T2 = (FTransform(SceneAsset->GetMeshToComponentForRole(OtherAnimTrack->Role).GetRotation()).Inverse() * (T2 * OtherAnimTrack->MeshToScene));
 
 								WarpPointTransform.SetLocation(FMath::Lerp<FVector>(T1.GetLocation(), T2.GetLocation(), Params.Weight));
 								WarpPointTransform.SetRotation((T2.GetLocation() - T1.GetLocation()).GetSafeNormal2D().ToOrientationQuat());
@@ -1416,7 +1415,7 @@ void FContextualAnimViewModel::CacheWarpPoints()
 						else
 						{
 							const FTransform RootTransform = AnimTrack->Animation ? UContextualAnimUtilities::ExtractRootTransformFromAnimation(AnimTrack->Animation, 0.f) : FTransform::Identity;
-							WarpPointTransform = (SceneAsset->GetMeshToComponentForRole(AnimTrack->Role).Inverse() * (RootTransform * AnimTrack->MeshToScene));
+							WarpPointTransform = (FTransform(SceneAsset->GetMeshToComponentForRole(AnimTrack->Role).GetRotation()).Inverse() * (RootTransform * AnimTrack->MeshToScene));
 						}
 					}
 
