Contextual Anim: Warp targets are now removed when the interaction ends

diff --git a/Engine/Plugins/Experimental/Animation/ContextualAnimation/Source/ContextualAnimation/Private/ContextualAnimSceneActorComponent.cpp b/Engine/Plugins/Experimental/Animation/ContextualAnimation/Source/ContextualAnimation/Private/ContextualAnimSceneActorComponent.cpp
--- a/Engine/Plugins/Experimental/Animation/ContextualAnimation/Source/ContextualAnimation/Private/ContextualAnimSceneActorComponent.cpp
+++ b/Engine/Plugins/Experimental/Animation/ContextualAnimation/Source/ContextualAnimation/Private/ContextualAnimSceneActorComponent.cpp
@@ -1003,6 +1003,11 @@ void UContextualAnimSceneActorComponent::LeaveScene()
 
 		RestoreMovementState(*Binding);
 
+		if (UMotionWarpingComponent* MotionWarpComp = Binding->GetMotionWarpingComponent())
+		{
+			MotionWarpComp->RemoveAllWarpTargets();
+		}
+
 		// Notify the other actors in the interaction
 		// @TODO: This should be refactored so only the leader of the interaction maintains the full bindings
 		for (const FContextualAnimSceneBinding& OtherBinding : Bindings)
