Contextual Anim - Downgraded verbosity of all the UE_LOGs from Log to Verbose

diff --git a/Engine/Plugins/Experimental/Animation/ContextualAnimation/Source/ContextualAnimation/Private/ContextualAnimSceneActorComponent.cpp b/Engine/Plugins/Experimental/Animation/ContextualAnimation/Source/ContextualAnimation/Private/ContextualAnimSceneActorComponent.cpp
--- a/Engine/Plugins/Experimental/Animation/ContextualAnimation/Source/ContextualAnimation/Private/ContextualAnimSceneActorComponent.cpp
+++ b/Engine/Plugins/Experimental/Animation/ContextualAnimation/Source/ContextualAnimation/Private/ContextualAnimSceneActorComponent.cpp
@@ -119,7 +119,7 @@ void UContextualAnimSceneActorComponent::PlayAnimation_Internal(UAnimSequenceBas
 {
 	if (UAnimInstance* AnimInstance = UContextualAnimUtilities::TryGetAnimInstance(GetOwner()))
 	{
-		UE_LOG(LogContextualAnim, Log, TEXT("%-21s \t\tUContextualAnimSceneActorComponent::PlayAnimation_Internal Playing Animation. Actor: %s Anim: %s StartTime: %f bSyncPlaybackTime: %d"),
+		UE_LOG(LogContextualAnim, Verbose, TEXT("%-21s \t\tUContextualAnimSceneActorComponent::PlayAnimation_Internal Playing Animation. Actor: %s Anim: %s StartTime: %f bSyncPlaybackTime: %d"),
 			*UEnum::GetValueAsString(TEXT("Engine.ENetRole"), GetOwner()->GetLocalRole()), *GetNameSafe(GetOwner()), *GetNameSafe(Animation), StartTime, bSyncPlaybackTime);
 
 		//@TODO: Add support for dynamic montage
@@ -243,7 +243,7 @@ bool UContextualAnimSceneActorComponent::LateJoinContextualAnimScene(AActor* Act
 		}
 	}
 
-	UE_LOG(LogContextualAnim, Log, TEXT("%-21s UContextualAnimSceneActorComponent::LateJoinContextualAnimScene Owner: %s Bindings Id: %d Section: %d Asset: %s. Requester: %s Role: %s"),
+	UE_LOG(LogContextualAnim, Verbose, TEXT("%-21s UContextualAnimSceneActorComponent::LateJoinContextualAnimScene Owner: %s Bindings Id: %d Section: %d Asset: %s. Requester: %s Role: %s"),
 		*UEnum::GetValueAsString(TEXT("Engine.ENetRole"), GetOwner()->GetLocalRole()), *GetNameSafe(GetOwner()), Bindings.GetID(), Bindings.GetSectionIdx(), *GetNameSafe(Bindings.GetSceneAsset()), *GetNameSafe(Actor), *Role.ToString());
 
 	// Add actor to the bindings
@@ -307,7 +307,7 @@ void UContextualAnimSceneActorComponent::LateJoinScene(const FContextualAnimScen
 
 	if (const FContextualAnimSceneBinding* Binding = InBindings.FindBindingByActor(GetOwner()))
 	{
-		UE_LOG(LogContextualAnim, Log, TEXT("%-21s UContextualAnimSceneActorComponent::LateJoinScene Actor: %s Role: %s Bindings Id: %d Section: %d Asset: %s"),
+		UE_LOG(LogContextualAnim, Verbose, TEXT("%-21s UContextualAnimSceneActorComponent::LateJoinScene Actor: %s Role: %s Bindings Id: %d Section: %d Asset: %s"),
 			*UEnum::GetValueAsString(TEXT("Engine.ENetRole"), GetOwner()->GetLocalRole()), *GetNameSafe(GetOwner()), *InBindings.GetRoleFromBinding(*Binding).ToString(), InBindings.GetID(), InBindings.GetSectionIdx(), *GetNameSafe(InBindings.GetSceneAsset()));
 
 		Bindings = InBindings;
@@ -336,7 +336,7 @@ void UContextualAnimSceneActorComponent::OnRep_LateJoinData()
 {
 	// This is received by the leader of the interaction on every remote client
 
-	UE_LOG(LogContextualAnim, Log, TEXT("%-21s UContextualAnimSceneActorComponent::OnRep_LateJoinData Owner: %s Bindings Id: %d Section: %d Asset: %s. Requester: %s Role: %s RepCounter: %d"),
+	UE_LOG(LogContextualAnim, Verbose, TEXT("%-21s UContextualAnimSceneActorComponent::OnRep_LateJoinData Owner: %s Bindings Id: %d Section: %d Asset: %s. Requester: %s Role: %s RepCounter: %d"),
 		*UEnum::GetValueAsString(TEXT("Engine.ENetRole"), GetOwner()->GetLocalRole()), *GetNameSafe(GetOwner()), Bindings.GetID(), Bindings.GetSectionIdx(), *GetNameSafe(Bindings.GetSceneAsset()), *GetNameSafe(RepLateJoinData.Actor), *RepLateJoinData.Role.ToString(), RepLateJoinData.RepCounter);
 
 	if (!RepLateJoinData.IsValid())
@@ -410,7 +410,7 @@ bool UContextualAnimSceneActorComponent::TransitionContextualAnimScene(FName Sec
 		const int32 SectionIdx = Bindings.GetSceneAsset()->GetSectionIndex(SectionName);
 		if (SectionIdx != INDEX_NONE)
 		{
-			UE_LOG(LogContextualAnim, Log, TEXT("%-21s UContextualAnimSceneActorComponent::TransitionTo Actor: %s SectionName: %s"),
+			UE_LOG(LogContextualAnim, Verbose, TEXT("%-21s UContextualAnimSceneActorComponent::TransitionTo Actor: %s SectionName: %s"),
 				*UEnum::GetValueAsString(TEXT("Engine.ENetRole"), GetOwner()->GetLocalRole()), *GetNameSafe(GetOwner()), *SectionName.ToString());
 
 			// Calculate WarpPoints
@@ -552,7 +552,7 @@ bool UContextualAnimSceneActorComponent::TransitionSingleActor(int32 SectionIdx,
 			const FContextualAnimTrack* AnimTrack = Asset->GetAnimTrack(SectionIdx, AnimSetIdx, Bindings.GetRoleFromBinding(*OwnerBinding));
 			if (AnimTrack && AnimTrack->Animation)
 			{
-				UE_LOG(LogContextualAnim, Log, TEXT("%-21s UContextualAnimSceneActorComponent::TransitionSingleActor Actor: %s SectionIdx: %d AnimSetIdx: %d"),
+				UE_LOG(LogContextualAnim, Verbose, TEXT("%-21s UContextualAnimSceneActorComponent::TransitionSingleActor Actor: %s SectionIdx: %d AnimSetIdx: %d"),
 					*UEnum::GetValueAsString(TEXT("Engine.ENetRole"), GetOwner()->GetLocalRole()), *GetNameSafe(GetOwner()), SectionIdx, AnimSetIdx);
 
 				// Calculate WarpPoints
@@ -597,7 +597,7 @@ void UContextualAnimSceneActorComponent::OnTransitionSingleActor(const FContextu
 
 void UContextualAnimSceneActorComponent::OnRep_RepTransitionSingleActor()
 {
-	UE_LOG(LogContextualAnim, Log, TEXT("%-21s UContextualAnimSceneActorComponent::OnRep_RepTransitionSingleActor Owner: %s Id: %d RepCounter: %d SectionIdx: %d AnimSetIdx: %d Current Bindings ID: %d"),
+	UE_LOG(LogContextualAnim, Verbose, TEXT("%-21s UContextualAnimSceneActorComponent::OnRep_RepTransitionSingleActor Owner: %s Id: %d RepCounter: %d SectionIdx: %d AnimSetIdx: %d Current Bindings ID: %d"),
 		*UEnum::GetValueAsString(TEXT("Engine.ENetRole"), GetOwner()->GetLocalRole()), *GetNameSafe(GetOwner()), 
 		RepTransitionSingleActorData.Id, RepTransitionSingleActorData.RepCounter, RepTransitionSingleActorData.SectionIdx, RepTransitionSingleActorData.AnimSetIdx, Bindings.IsValid() ? Bindings.GetID() : -1);
 
@@ -667,7 +667,7 @@ bool UContextualAnimSceneActorComponent::TransitionSingleActor(int32 SectionIdx,
 
 bool UContextualAnimSceneActorComponent::StartContextualAnimScene(const FContextualAnimSceneBindings& InBindings, const TArray<FContextualAnimWarpTarget>& ExternalWarpTargets)
 {
-	UE_LOG(LogContextualAnim, Log, TEXT("%-21s UContextualAnimSceneActorComponent::StartContextualAnim Actor: %s"),
+	UE_LOG(LogContextualAnim, Verbose, TEXT("%-21s UContextualAnimSceneActorComponent::StartContextualAnim Actor: %s"),
 		*UEnum::GetValueAsString(TEXT("Engine.ENetRole"), GetOwner()->GetLocalRole()), *GetNameSafe(GetOwner()));
 
 	const FContextualAnimSceneBinding* OwnerBinding = InBindings.FindBindingByActor(GetOwner());
@@ -785,7 +785,7 @@ bool UContextualAnimSceneActorComponent::ServerEarlyOutContextualAnimScene_Valid
 
 void UContextualAnimSceneActorComponent::OnRep_TransitionData()
 {
-	UE_LOG(LogContextualAnim, Log, TEXT("%-21s UContextualAnimSceneActorComponent::OnRep_TransitionData Actor: %s SectionIdx: %d AnimsetIdx: %d RepCounter: %d"),
+	UE_LOG(LogContextualAnim, Verbose, TEXT("%-21s UContextualAnimSceneActorComponent::OnRep_TransitionData Actor: %s SectionIdx: %d AnimsetIdx: %d RepCounter: %d"),
 		*UEnum::GetValueAsString(TEXT("Engine.ENetRole"), GetOwner()->GetLocalRole()), *GetNameSafe(GetOwner()), RepTransitionData.SectionIdx, RepTransitionData.AnimSetIdx, RepTransitionData.RepCounter);
 
 	if (!RepTransitionData.IsValid())
@@ -806,7 +806,7 @@ void UContextualAnimSceneActorComponent::OnRep_TransitionData()
 
 void UContextualAnimSceneActorComponent::OnRep_Bindings()
 {
-	UE_LOG(LogContextualAnim, Log, TEXT("%-21s UContextualAnimSceneActorComponent::OnRep_Bindings Actor: %s Rep Bindings Id: %d RepCounter: %d Num: %d Current Bindings Id: %d Num: %d"),
+	UE_LOG(LogContextualAnim, Verbose, TEXT("%-21s UContextualAnimSceneActorComponent::OnRep_Bindings Actor: %s Rep Bindings Id: %d RepCounter: %d Num: %d Current Bindings Id: %d Num: %d"),
 		*UEnum::GetValueAsString(TEXT("Engine.ENetRole"), GetOwner()->GetLocalRole()), *GetNameSafe(GetOwner()), RepBindings.Bindings.GetID(), RepBindings.RepCounter, RepBindings.Bindings.Num(), Bindings.GetID(), Bindings.Num());
 
 	if (!RepBindings.IsValid())
@@ -943,7 +943,7 @@ void UContextualAnimSceneActorComponent::JoinScene(const FContextualAnimSceneBin
 
 	if (const FContextualAnimSceneBinding* Binding = InBindings.FindBindingByActor(GetOwner()))
 	{
-		UE_LOG(LogContextualAnim, Log, TEXT("%-21s UContextualAnimSceneActorComponent::JoinScene Actor: %s Role: %s InBindings Id: %d Section: %d Asset: %s"),
+		UE_LOG(LogContextualAnim, Verbose, TEXT("%-21s UContextualAnimSceneActorComponent::JoinScene Actor: %s Role: %s InBindings Id: %d Section: %d Asset: %s"),
 			*UEnum::GetValueAsString(TEXT("Engine.ENetRole"), GetOwner()->GetLocalRole()), *GetNameSafe(GetOwner()), *InBindings.GetRoleFromBinding(*Binding).ToString(), InBindings.GetID(), InBindings.GetSectionIdx(), *GetNameSafe(InBindings.GetSceneAsset()));
 
 		AnimsPlayed.Reset();
@@ -974,7 +974,7 @@ void UContextualAnimSceneActorComponent::LeaveScene()
 {
 	if (const FContextualAnimSceneBinding* Binding = Bindings.FindBindingByActor(GetOwner()))
 	{
-		UE_LOG(LogContextualAnim, Log, TEXT("%-21s UContextualAnimSceneActorComponent::LeaveScene Actor: %s Role: %s Current Bindings Id: %d Section: %d Asset: %s"),
+		UE_LOG(LogContextualAnim, Verbose, TEXT("%-21s UContextualAnimSceneActorComponent::LeaveScene Actor: %s Role: %s Current Bindings Id: %d Section: %d Asset: %s"),
 			*UEnum::GetValueAsString(TEXT("Engine.ENetRole"), GetOwner()->GetLocalRole()), *GetNameSafe(GetOwner()), *Bindings.GetRoleFromBinding(*Binding).ToString(),
 			Bindings.GetID(), Bindings.GetSectionIdx(), *GetNameSafe(Bindings.GetSceneAsset()));
 
@@ -1100,7 +1100,7 @@ bool UContextualAnimSceneActorComponent::CanLeaveScene(const FContextualAnimScen
 
 void UContextualAnimSceneActorComponent::OnMontageBlendingOut(UAnimMontage* Montage, bool bInterrupted)
 {
-	UE_LOG(LogContextualAnim, Log, TEXT("%-21s UContextualAnimSceneActorComponent::OnMontageBlendingOut Actor: %s Montage: %s bInterrupted: %d"),
+	UE_LOG(LogContextualAnim, Verbose, TEXT("%-21s UContextualAnimSceneActorComponent::OnMontageBlendingOut Actor: %s Montage: %s bInterrupted: %d"),
 		*UEnum::GetValueAsString(TEXT("Engine.ENetRole"), GetOwner()->GetLocalRole()), *GetNameSafe(GetOwner()), *GetNameSafe(Montage), bInterrupted);
 
 	if (const FContextualAnimSceneBinding* Binding = Bindings.FindBindingByActor(GetOwner()))
@@ -1108,13 +1108,13 @@ void UContextualAnimSceneActorComponent::OnMontageBlendingOut(UAnimMontage* Mont
 		AnimsPlayed.RemoveSingleSwap(Montage);
 		if (AnimsPlayed.Num() > 0)
 		{
-			UE_LOG(LogContextualAnim, Log, TEXT("%-21s \tUContextualAnimSceneActorComponent::OnMontageBlendingOut AnimsPlayed Num: %d"), *UEnum::GetValueAsString(TEXT("Engine.ENetRole"), GetOwner()->GetLocalRole()), AnimsPlayed.Num());
+			UE_LOG(LogContextualAnim, Verbose, TEXT("%-21s \tUContextualAnimSceneActorComponent::OnMontageBlendingOut AnimsPlayed Num: %d"), *UEnum::GetValueAsString(TEXT("Engine.ENetRole"), GetOwner()->GetLocalRole()), AnimsPlayed.Num());
 			return;
 		}
 
 		if (!CanLeaveScene(*Binding))
 		{
-			UE_LOG(LogContextualAnim, Log, TEXT("%-21s \tUContextualAnimSceneActorComponent::OnMontageBlendingOut CanLeaveScene FALSE"), *UEnum::GetValueAsString(TEXT("Engine.ENetRole"), GetOwner()->GetLocalRole()));
+			UE_LOG(LogContextualAnim, Verbose, TEXT("%-21s \tUContextualAnimSceneActorComponent::OnMontageBlendingOut CanLeaveScene FALSE"), *UEnum::GetValueAsString(TEXT("Engine.ENetRole"), GetOwner()->GetLocalRole()));
 			return;
 		}
 
