[Mover] - Furniture Support - Beds
Mover Pawns can now sleep on the beds

diff --git a/Engine/Plugins/Experimental/Animation/ContextualAnimation/Source/ContextualAnimation/Private/ContextualAnimSceneActorComponent.cpp b/Engine/Plugins/Experimental/Animation/ContextualAnimation/Source/ContextualAnimation/Private/ContextualAnimSceneActorComponent.cpp
--- a/Engine/Plugins/Experimental/Animation/ContextualAnimation/Source/ContextualAnimation/Private/ContextualAnimSceneActorComponent.cpp
+++ b/Engine/Plugins/Experimental/Animation/ContextualAnimation/Source/ContextualAnimation/Private/ContextualAnimSceneActorComponent.cpp
@@ -1064,6 +1064,7 @@ void UContextualAnimSceneActorComponent::SetMovementState(const FContextualAnimS
 		CharacterPropertiesBackup.bUseControllerDesiredRotation = MovementComp->bUseControllerDesiredRotation;
 		CharacterPropertiesBackup.bOrientRotationToMovement = MovementComp->bOrientRotationToMovement;
 		CharacterPropertiesBackup.MovementMode = MovementComp->MovementMode;
+		CharacterPropertiesBackup.bSimulatePhysics = MovementComp->UpdatedPrimitive && MovementComp->UpdatedPrimitive->IsSimulatingPhysics();
 
 		// Disable movement correction.
 		MovementComp->bIgnoreClientMovementErrorChecksAndCorrection = true;
@@ -1073,6 +1074,11 @@ void UContextualAnimSceneActorComponent::SetMovementState(const FContextualAnimS
 		MovementComp->bUseControllerDesiredRotation = false;
 		MovementComp->bOrientRotationToMovement = false;
 
+		if (UPrimitiveComponent* const UpdatedPrimitive = MovementComp->UpdatedPrimitive)
+		{
+			UpdatedPrimitive->SetSimulatePhysics(false);
+		}
+
 		if (MovementComp->MovementMode != DesiredMoveMode)
 		{
 			MovementComp->SetMovementMode(DesiredMoveMode);
@@ -1090,6 +1096,10 @@ void UContextualAnimSceneActorComponent::RestoreMovementState(const FContextualA
 		MovementComp->bUseControllerDesiredRotation = CharacterPropertiesBackup.bUseControllerDesiredRotation;
 		MovementComp->bOrientRotationToMovement = CharacterPropertiesBackup.bOrientRotationToMovement;
 		MovementComp->SetMovementMode(CharacterPropertiesBackup.MovementMode);
+		if(UPrimitiveComponent* const UpdatedPrimitive = MovementComp->UpdatedPrimitive)
+		{
+			UpdatedPrimitive->SetSimulatePhysics(CharacterPropertiesBackup.bSimulatePhysics);
+		}
 	}
 }
 
diff --git a/Engine/Plugins/Experimental/Animation/ContextualAnimation/Source/ContextualAnimation/Public/ContextualAnimSceneActorComponent.h b/Engine/Plugins/Experimental/Animation/ContextualAnimation/Source/ContextualAnimation/Public/ContextualAnimSceneActorComponent.h
--- a/Engine/Plugins/Experimental/Animation/ContextualAnimation/Source/ContextualAnimation/Public/ContextualAnimSceneActorComponent.h
+++ b/Engine/Plugins/Experimental/Animation/ContextualAnimation/Source/ContextualAnimation/Public/ContextualAnimSceneActorComponent.h
@@ -254,6 +254,7 @@ protected:
 		bool bAllowPhysicsRotationDuringAnimRootMotion = false;
 		bool bUseControllerDesiredRotation = false;
 		bool bOrientRotationToMovement = false;
+		bool bSimulatePhysics = false;
 		EMovementMode MovementMode = EMovementMode::MOVE_Walking;
 		TArray<TTuple<ECollisionChannel, ECollisionResponse>> CollisionResponses;
 	};
