Cleanup scale after a CAS scene is completed:
- fixes a bug with chairs where sometimes the scaling gets set to (-1, 1, 1) causing bugs in multiple other systems such as the character movement component staying in fall mode forever because of the scaled capsule radius becoming negative

diff --git a/Engine/Plugins/Experimental/Animation/ContextualAnimation/Source/ContextualAnimation/Private/ContextualAnimSceneActorComponent.cpp b/Engine/Plugins/Experimental/Animation/ContextualAnimation/Source/ContextualAnimation/Private/ContextualAnimSceneActorComponent.cpp
--- a/Engine/Plugins/Experimental/Animation/ContextualAnimation/Source/ContextualAnimation/Private/ContextualAnimSceneActorComponent.cpp
+++ b/Engine/Plugins/Experimental/Animation/ContextualAnimation/Source/ContextualAnimation/Private/ContextualAnimSceneActorComponent.cpp
@@ -974,6 +974,8 @@ void UContextualAnimSceneActorComponent::LeaveScene()
 {
 	if (const FContextualAnimSceneBinding* Binding = Bindings.FindBindingByActor(GetOwner()))
 	{
+		OnPreLeaveScene(*Binding);
+		
 		UE_LOG(LogContextualAnim, Verbose, TEXT("%-21s UContextualAnimSceneActorComponent::LeaveScene Actor: %s Role: %s Current Bindings Id: %d Section: %d Asset: %s"),
 			*UEnum::GetValueAsString(TEXT("Engine.ENetRole"), GetOwner()->GetLocalRole()), *GetNameSafe(GetOwner()), *Bindings.GetRoleFromBinding(*Binding).ToString(),
 			Bindings.GetID(), Bindings.GetSectionIdx(), *GetNameSafe(Bindings.GetSceneAsset()));
@@ -1108,6 +1110,11 @@ bool UContextualAnimSceneActorComponent::CanLeaveScene(const FContextualAnimScen
 	return true;
 }
 
+void UContextualAnimSceneActorComponent::OnPreLeaveScene(const FContextualAnimSceneBinding& Binding)
+{
+	// For derived classes to override.
+}
+
 void UContextualAnimSceneActorComponent::OnMontageBlendingOut(UAnimMontage* Montage, bool bInterrupted)
 {
 	UE_LOG(LogContextualAnim, Verbose, TEXT("%-21s UContextualAnimSceneActorComponent::OnMontageBlendingOut Actor: %s Montage: %s bInterrupted: %d"),
diff --git a/Engine/Plugins/Experimental/Animation/ContextualAnimation/Source/ContextualAnimation/Public/ContextualAnimSceneActorComponent.h b/Engine/Plugins/Experimental/Animation/ContextualAnimation/Source/ContextualAnimation/Public/ContextualAnimSceneActorComponent.h
--- a/Engine/Plugins/Experimental/Animation/ContextualAnimation/Source/ContextualAnimation/Public/ContextualAnimSceneActorComponent.h
+++ b/Engine/Plugins/Experimental/Animation/ContextualAnimation/Source/ContextualAnimation/Public/ContextualAnimSceneActorComponent.h
@@ -307,6 +307,9 @@ protected:
 
 	virtual bool CanLeaveScene(const FContextualAnimSceneBinding& Binding);
 
+	// Called just before the scene leaving happens (which calls OnLeaveScene)
+	virtual void OnPreLeaveScene(const FContextualAnimSceneBinding& Binding);
+
 	virtual void OnLeaveScene(const FContextualAnimSceneBinding& Binding);
 
 	void LateJoinScene(const FContextualAnimSceneBindings& InBindings, int32 SectionIdx, int32 AnimSetIdx, const TArray<FContextualAnimWarpPoint>& WarpPoints, const TArray<FContextualAnimWarpTarget>& ExternalWarpTargets);
