Contextual Anim -  Made OtherActorLeftScene function virtual

diff --git a/Engine/Plugins/Experimental/Animation/ContextualAnimation/Source/ContextualAnimation/Public/ContextualAnimSceneActorComponent.h b/Engine/Plugins/Experimental/Animation/ContextualAnimation/Source/ContextualAnimation/Public/ContextualAnimSceneActorComponent.h
--- a/Engine/Plugins/Experimental/Animation/ContextualAnimation/Source/ContextualAnimation/Public/ContextualAnimSceneActorComponent.h
+++ b/Engine/Plugins/Experimental/Animation/ContextualAnimation/Source/ContextualAnimation/Public/ContextualAnimSceneActorComponent.h
@@ -334,7 +334,7 @@ protected:
 
 	virtual void OnTransitionSingleActor(const FContextualAnimSceneBinding& Binding, int32 SectionIdx, int32 AnimSetIdx);
 
-	void OtherActorLeftScene(AActor& Actor);
+	virtual void OtherActorLeftScene(AActor& Actor);
 
 private:
 
