Contextual Anim - Changed the logic for cleaning up warp targets to only remove the ones added during the interaction.

diff --git a/Engine/Plugins/Experimental/Animation/ContextualAnimation/Source/ContextualAnimation/Private/ContextualAnimSceneActorComponent.cpp b/Engine/Plugins/Experimental/Animation/ContextualAnimation/Source/ContextualAnimation/Private/ContextualAnimSceneActorComponent.cpp
--- a/Engine/Plugins/Experimental/Animation/ContextualAnimation/Source/ContextualAnimation/Private/ContextualAnimSceneActorComponent.cpp
+++ b/Engine/Plugins/Experimental/Animation/ContextualAnimation/Source/ContextualAnimation/Private/ContextualAnimSceneActorComponent.cpp
@@ -202,6 +202,7 @@ void UContextualAnimSceneActorComponent::AddOrUpdateWarpTargets(int32 SectionIdx
 					const FTransform TransformRelativeToWarpPoint = Asset->GetAlignmentTransform(*AnimTrack, WarpPoint.Name, Time);
 					const FTransform WarpTargetTransform = TransformRelativeToWarpPoint * WarpPoint.Transform;
 					MotionWarpComp->AddOrUpdateWarpTargetFromTransform(WarpPoint.Name, WarpTargetTransform);
+					WarpTargetNamesCache.AddUnique(WarpPoint.Name);
 				}
 			}
 		}
@@ -212,6 +213,7 @@ void UContextualAnimSceneActorComponent::AddOrUpdateWarpTargets(int32 SectionIdx
 			if (WarpTarget.Role == Role)
 			{
 				MotionWarpComp->AddOrUpdateWarpTargetFromTransform(WarpTarget.TargetName, FTransform(WarpTarget.TargetRotation, WarpTarget.TargetLocation));
+				WarpTargetNamesCache.AddUnique(WarpTarget.TargetName);
 			}
 		}
 	}
@@ -1005,9 +1007,14 @@ void UContextualAnimSceneActorComponent::LeaveScene()
 
 		RestoreMovementState(*Binding);
 
-		if (UMotionWarpingComponent* MotionWarpComp = Binding->GetMotionWarpingComponent())
+		if (WarpTargetNamesCache.Num() > 0)
 		{
-			MotionWarpComp->RemoveAllWarpTargets();
+			if (UMotionWarpingComponent* MotionWarpComp = Binding->GetMotionWarpingComponent())
+			{
+				MotionWarpComp->RemoveWarpTargets(WarpTargetNamesCache);
+			}
+
+			WarpTargetNamesCache.Reset();
 		}
 
 		// Notify the other actors in the interaction
diff --git a/Engine/Plugins/Experimental/Animation/ContextualAnimation/Source/ContextualAnimation/Public/ContextualAnimSceneActorComponent.h b/Engine/Plugins/Experimental/Animation/ContextualAnimation/Source/ContextualAnimation/Public/ContextualAnimSceneActorComponent.h
--- a/Engine/Plugins/Experimental/Animation/ContextualAnimation/Source/ContextualAnimation/Public/ContextualAnimSceneActorComponent.h
+++ b/Engine/Plugins/Experimental/Animation/ContextualAnimation/Source/ContextualAnimation/Public/ContextualAnimSceneActorComponent.h
@@ -248,6 +248,13 @@ protected:
 	UPROPERTY(Transient)
 	TArray<FContextualAnimIKTarget> IKTargets;
 
+	/** 
+	 * Motion warping doesn't remove warp targets when a warping window ends, so we keep track of the warp targets added during the interaction 
+	 * and remove them at the end of it to prevent issues caused by dangling warp targets after an interaction.
+	 */
+	UPROPERTY(Transient)
+	TArray<FName> WarpTargetNamesCache;
+
 	struct FCharacterProperties
 	{
 		bool bIgnoreClientMovementErrorChecksAndCorrection = false;
