Contextual Anim - Added option for setting custom movement mode during an interaction (using numerical value for now until we add custom UI to expose the list custom modes from the game-specific enum.)

diff --git a/Engine/Plugins/Experimental/Animation/ContextualAnimation/Source/ContextualAnimation/Private/ContextualAnimSceneActorComponent.cpp b/Engine/Plugins/Experimental/Animation/ContextualAnimation/Source/ContextualAnimation/Private/ContextualAnimSceneActorComponent.cpp
--- a/Engine/Plugins/Experimental/Animation/ContextualAnimation/Source/ContextualAnimation/Private/ContextualAnimSceneActorComponent.cpp
+++ b/Engine/Plugins/Experimental/Animation/ContextualAnimation/Source/ContextualAnimation/Private/ContextualAnimSceneActorComponent.cpp
@@ -323,7 +323,7 @@ void UContextualAnimSceneActorComponent::LateJoinScene(const FContextualAnimScen
 
 		SetCollisionState(*Binding);
 
-		SetMovementState(*Binding, AnimTrack->MovementMode);
+		SetMovementState(*Binding, AnimTrack->MovementMode, AnimTrack->CustomMovementMode);
 
 		OnLateJoinScene(*Binding, SectionIdx, AnimSetIdx);
 	}
@@ -959,7 +959,7 @@ void UContextualAnimSceneActorComponent::JoinScene(const FContextualAnimSceneBin
 
 		SetCollisionState(*Binding);
 
-		SetMovementState(*Binding, AnimTrack.MovementMode);
+		SetMovementState(*Binding, AnimTrack.MovementMode, AnimTrack.CustomMovementMode);
 
 		OnJoinScene(*Binding);
 
@@ -1063,7 +1063,7 @@ void UContextualAnimSceneActorComponent::OnLeaveScene(const FContextualAnimScene
 	// For derived classes to override.
 }
 
-void UContextualAnimSceneActorComponent::SetMovementState(const FContextualAnimSceneBinding& Binding, EMovementMode DesiredMoveMode)
+void UContextualAnimSceneActorComponent::SetMovementState(const FContextualAnimSceneBinding& Binding, EMovementMode DesiredMoveMode, uint8 CustomMode)
 {
 	if (UCharacterMovementComponent* MovementComp = Binding.GetCharacterMovementComponent())
 	{
@@ -1090,7 +1090,7 @@ void UContextualAnimSceneActorComponent::SetMovementState(const FContextualAnimS
 
 		if (MovementComp->MovementMode != DesiredMoveMode)
 		{
-			MovementComp->SetMovementMode(DesiredMoveMode);
+			MovementComp->SetMovementMode(DesiredMoveMode, CustomMode);
 		}
 	}
 }
diff --git a/Engine/Plugins/Experimental/Animation/ContextualAnimation/Source/ContextualAnimation/Public/ContextualAnimSceneActorComponent.h b/Engine/Plugins/Experimental/Animation/ContextualAnimation/Source/ContextualAnimation/Public/ContextualAnimSceneActorComponent.h
--- a/Engine/Plugins/Experimental/Animation/ContextualAnimation/Source/ContextualAnimation/Public/ContextualAnimSceneActorComponent.h
+++ b/Engine/Plugins/Experimental/Animation/ContextualAnimation/Source/ContextualAnimation/Public/ContextualAnimSceneActorComponent.h
@@ -294,7 +294,7 @@ protected:
 	virtual void SetCollisionState(const FContextualAnimSceneBinding& Binding);
 	virtual void RestoreCollisionState(const FContextualAnimSceneBinding& Binding);
 
-	void SetMovementState(const FContextualAnimSceneBinding& Binding, EMovementMode DesiredMoveMode);
+	void SetMovementState(const FContextualAnimSceneBinding& Binding, EMovementMode DesiredMoveMode, uint8 CustomMode = 0);
 	void RestoreMovementState(const FContextualAnimSceneBinding& Binding);
 
 	UFUNCTION()
diff --git a/Engine/Plugins/Experimental/Animation/ContextualAnimation/Source/ContextualAnimation/Public/ContextualAnimTypes.h b/Engine/Plugins/Experimental/Animation/ContextualAnimation/Source/ContextualAnimation/Public/ContextualAnimTypes.h
--- a/Engine/Plugins/Experimental/Animation/ContextualAnimation/Source/ContextualAnimation/Public/ContextualAnimTypes.h
+++ b/Engine/Plugins/Experimental/Animation/ContextualAnimation/Source/ContextualAnimation/Public/ContextualAnimTypes.h
@@ -77,6 +77,9 @@ struct CONTEXTUALANIMATION_API FContextualAnimTrack
 	UPROPERTY(EditAnywhere, BlueprintReadWrite, Category = "Defaults")
 	TEnumAsByte<enum EMovementMode> MovementMode = EMovementMode::MOVE_Walking;
 
+	UPROPERTY(EditAnywhere, BlueprintReadWrite, Category = "Defaults", meta = (EditConditionHides, EditCondition = "MovementMode==EMovementMode::MOVE_Custom"))
+	uint8 CustomMovementMode = 0;
+
 	/** Whether the actor that should play this animation is optional */
 	UPROPERTY(EditAnywhere, BlueprintReadWrite, Category = "Defaults")
 	bool bOptional = false;
diff --git a/Engine/Plugins/Experimental/Animation/ContextualAnimation/Source/ContextualAnimationEditor/Private/ContextualAnimViewModel.cpp b/Engine/Plugins/Experimental/Animation/ContextualAnimation/Source/ContextualAnimationEditor/Private/ContextualAnimViewModel.cpp
--- a/Engine/Plugins/Experimental/Animation/ContextualAnimation/Source/ContextualAnimationEditor/Private/ContextualAnimViewModel.cpp
+++ b/Engine/Plugins/Experimental/Animation/ContextualAnimation/Source/ContextualAnimationEditor/Private/ContextualAnimViewModel.cpp
@@ -596,6 +596,7 @@ void FContextualAnimViewModel::AddNewAnimSet(const FContextualAnimNewAnimSetPara
 		AnimTrack.Role = Data.RoleName;
 		AnimTrack.Animation = Data.Animation;
 		AnimTrack.MovementMode = Data.MovementMode;
+		AnimTrack.CustomMovementMode = Data.CustomMovementMode;
 		AnimTrack.bOptional = Data.bOptional;
 		AnimSet.Tracks.Add(AnimTrack);
 		AnimSet.RandomWeight = Params.RandomWeight;
diff --git a/Engine/Plugins/Experimental/Animation/ContextualAnimation/Source/ContextualAnimationEditor/Public/ContextualAnimEditorTypes.h b/Engine/Plugins/Experimental/Animation/ContextualAnimation/Source/ContextualAnimationEditor/Public/ContextualAnimEditorTypes.h
--- a/Engine/Plugins/Experimental/Animation/ContextualAnimation/Source/ContextualAnimationEditor/Public/ContextualAnimEditorTypes.h
+++ b/Engine/Plugins/Experimental/Animation/ContextualAnimation/Source/ContextualAnimationEditor/Public/ContextualAnimEditorTypes.h
@@ -77,6 +77,9 @@ struct FContextualAnimNewAnimSetData
 	UPROPERTY(EditAnywhere, Category = "Settings")
 	TEnumAsByte<enum EMovementMode> MovementMode = EMovementMode::MOVE_Walking;
 
+	UPROPERTY(EditAnywhere, Category = "Settings", meta = (EditConditionHides, EditCondition = "MovementMode==EMovementMode::MOVE_Custom"))
+	uint8 CustomMovementMode = 0;
+
 	UPROPERTY(EditAnywhere, Category = "Settings")
 	bool bOptional = false;
 
