Contextual Anim - Added option to disable animation syncing per section

diff --git a/Engine/Plugins/Experimental/Animation/ContextualAnimation/Source/ContextualAnimation/Private/ContextualAnimSceneActorComponent.cpp b/Engine/Plugins/Experimental/Animation/ContextualAnimation/Source/ContextualAnimation/Private/ContextualAnimSceneActorComponent.cpp
--- a/Engine/Plugins/Experimental/Animation/ContextualAnimation/Source/ContextualAnimation/Private/ContextualAnimSceneActorComponent.cpp
+++ b/Engine/Plugins/Experimental/Animation/ContextualAnimation/Source/ContextualAnimation/Private/ContextualAnimSceneActorComponent.cpp
@@ -500,7 +500,7 @@ void UContextualAnimSceneActorComponent::HandleTransitionSelf(int32 NewSectionId
 	// Play animation
 	const FContextualAnimSceneBinding& Binding = *Bindings.FindBindingByActor(GetOwner());
 	const FContextualAnimTrack& AnimTrack = Bindings.GetAnimTrackFromBinding(Binding);
-	PlayAnimation_Internal(AnimTrack.Animation, 0.f, true);
+	PlayAnimation_Internal(AnimTrack.Animation, 0.f, Bindings.ShouldSyncAnimation());
 
 	AddOrUpdateWarpTargets(NewSectionIdx, NewAnimSetIdx, WarpPoints, ExternalWarpTargets);
 
@@ -953,7 +953,7 @@ void UContextualAnimSceneActorComponent::JoinScene(const FContextualAnimSceneBin
 		Bindings = InBindings;
 
 		const FContextualAnimTrack& AnimTrack = Bindings.GetAnimTrackFromBinding(*Binding);
-		PlayAnimation_Internal(AnimTrack.Animation, 0.f, true);
+		PlayAnimation_Internal(AnimTrack.Animation, 0.f, Bindings.ShouldSyncAnimation());
 
 		AddOrUpdateWarpTargets(AnimTrack.SectionIdx, AnimTrack.AnimSetIdx, WarpPoints, ExternalWarpTargets);
 
diff --git a/Engine/Plugins/Experimental/Animation/ContextualAnimation/Source/ContextualAnimation/Private/ContextualAnimTypes.cpp b/Engine/Plugins/Experimental/Animation/ContextualAnimation/Source/ContextualAnimation/Private/ContextualAnimTypes.cpp
--- a/Engine/Plugins/Experimental/Animation/ContextualAnimation/Source/ContextualAnimation/Private/ContextualAnimTypes.cpp
+++ b/Engine/Plugins/Experimental/Animation/ContextualAnimation/Source/ContextualAnimation/Private/ContextualAnimTypes.cpp
@@ -405,6 +405,12 @@ void FContextualAnimSceneBindings::Clear()
 	Data.Reset();
 }
 
+bool FContextualAnimSceneBindings::ShouldSyncAnimation() const
+{
+	const FContextualAnimSceneSection* Section = IsValid() ? SceneAsset->GetSection(SectionIdx) : nullptr;
+	return Section ? Section->ShouldSyncAnimations() : false;
+}
+
 const FContextualAnimSceneBinding* FContextualAnimSceneBindings::GetSyncLeader() const
 {
 	//@TODO: Return first secondary binding as sync leader for now. This may have to be explicitly defined, either in the SceneAsset or when creating the bindings.
diff --git a/Engine/Plugins/Experimental/Animation/ContextualAnimation/Source/ContextualAnimation/Public/ContextualAnimSceneAsset.h b/Engine/Plugins/Experimental/Animation/ContextualAnimation/Source/ContextualAnimation/Public/ContextualAnimSceneAsset.h
--- a/Engine/Plugins/Experimental/Animation/ContextualAnimation/Source/ContextualAnimation/Public/ContextualAnimSceneAsset.h
+++ b/Engine/Plugins/Experimental/Animation/ContextualAnimation/Source/ContextualAnimation/Public/ContextualAnimSceneAsset.h
@@ -112,6 +112,7 @@ public:
 	FORCEINLINE FName GetName() const { return Name; }
 	FORCEINLINE const TArray<FContextualAnimWarpPointDefinition>& GetWarpPointDefinitions() const { return WarpPointDefinitions; }
 	FORCEINLINE int32 GetNumAnimSets() const { return AnimSets.Num(); }
+	FORCEINLINE bool ShouldSyncAnimations() const { return bSyncAnimations; }
 
 protected:
 
@@ -127,6 +128,9 @@ protected:
 	UPROPERTY(EditAnywhere, Category = "Defaults", meta = (TitleProperty = "WarpTargetName"))
 	TArray<FContextualAnimWarpPointDefinition> WarpPointDefinitions;
 
+	UPROPERTY(EditAnywhere, Category = "Defaults")
+	bool bSyncAnimations = true;
+
 	void GenerateAlignmentTracks(UContextualAnimSceneAsset& SceneAsset);
 	void GenerateIKTargetTracks(UContextualAnimSceneAsset& SceneAsset);
 
diff --git a/Engine/Plugins/Experimental/Animation/ContextualAnimation/Source/ContextualAnimation/Public/ContextualAnimTypes.h b/Engine/Plugins/Experimental/Animation/ContextualAnimation/Source/ContextualAnimation/Public/ContextualAnimTypes.h
--- a/Engine/Plugins/Experimental/Animation/ContextualAnimation/Source/ContextualAnimation/Public/ContextualAnimTypes.h
+++ b/Engine/Plugins/Experimental/Animation/ContextualAnimation/Source/ContextualAnimation/Public/ContextualAnimTypes.h
@@ -498,6 +498,7 @@ struct CONTEXTUALANIMATION_API FContextualAnimSceneBindings
 	const FContextualAnimIKTargetDefContainer& GetIKTargetDefContainerFromBinding(const FContextualAnimSceneBinding& Binding) const;
 	FTransform GetIKTargetTransformFromBinding(const FContextualAnimSceneBinding& Binding, const FName& TrackName, float Time) const;
 
+	bool ShouldSyncAnimation() const;
 	const FContextualAnimSceneBinding* GetSyncLeader() const;
 	const FContextualAnimSceneBinding* GetPrimaryBinding() const;
 
diff --git a/Engine/Plugins/Experimental/Animation/ContextualAnimation/Source/ContextualAnimationEditor/Private/DetailCustomizations/ContextualAnimSceneAssetDetailCustom.cpp b/Engine/Plugins/Experimental/Animation/ContextualAnimation/Source/ContextualAnimationEditor/Private/DetailCustomizations/ContextualAnimSceneAssetDetailCustom.cpp
--- a/Engine/Plugins/Experimental/Animation/ContextualAnimation/Source/ContextualAnimationEditor/Private/DetailCustomizations/ContextualAnimSceneAssetDetailCustom.cpp
+++ b/Engine/Plugins/Experimental/Animation/ContextualAnimation/Source/ContextualAnimationEditor/Private/DetailCustomizations/ContextualAnimSceneAssetDetailCustom.cpp
@@ -75,6 +75,7 @@ void FContextualAnimSceneAssetDetailCustom::CustomizeDetails(IDetailLayoutBuilde
 		CurrentSectionCategory.AddProperty(SectionPropertyHandle->GetChildHandle(FName(TEXT("Name"))).ToSharedRef());
 		CurrentSectionCategory.AddProperty(SectionPropertyHandle->GetChildHandle(FName(TEXT("RoleToIKTargetDefsMap"))).ToSharedRef());
 		CurrentSectionCategory.AddProperty(SectionPropertyHandle->GetChildHandle(FName(TEXT("WarpPointDefinitions"))).ToSharedRef());
+		CurrentSectionCategory.AddProperty(SectionPropertyHandle->GetChildHandle(FName(TEXT("bSyncAnimations"))).ToSharedRef());
 
 		// Add button to generate warp points for the scene
 		CurrentSectionCategory.AddCustomRow(FText::GetEmpty())
