Contextual Anim - Exposed GetBindings to blueprints

diff --git a/Engine/Plugins/Experimental/Animation/ContextualAnimation/Source/ContextualAnimation/Public/ContextualAnimSceneActorComponent.h b/Engine/Plugins/Experimental/Animation/ContextualAnimation/Source/ContextualAnimation/Public/ContextualAnimSceneActorComponent.h
--- a/Engine/Plugins/Experimental/Animation/ContextualAnimation/Source/ContextualAnimation/Public/ContextualAnimSceneActorComponent.h
+++ b/Engine/Plugins/Experimental/Animation/ContextualAnimation/Source/ContextualAnimation/Public/ContextualAnimSceneActorComponent.h
@@ -176,6 +176,7 @@ public:
 
 	virtual void AddIKGoals_Implementation(TMap<FName, FIKRigGoal>& OutGoals) override;
 
+	UFUNCTION(BlueprintPure, Category = "Contextual Anim|Scene Actor Component")
 	const FContextualAnimSceneBindings& GetBindings() const { return Bindings; };
 
 	UFUNCTION(BlueprintPure, Category = "Contextual Anim|Scene Actor Component")
