Contextual Anim - Exposed to blueprint function for transitioning to a specific set + Exposed IgnoreClientMovementErrorChecksAndCorrection as param on the cas asset.

diff --git a/Engine/Plugins/Experimental/Animation/ContextualAnimation/Source/ContextualAnimation/Private/ContextualAnimSceneActorComponent.cpp b/Engine/Plugins/Experimental/Animation/ContextualAnimation/Source/ContextualAnimation/Private/ContextualAnimSceneActorComponent.cpp
--- a/Engine/Plugins/Experimental/Animation/ContextualAnimation/Source/ContextualAnimation/Private/ContextualAnimSceneActorComponent.cpp
+++ b/Engine/Plugins/Experimental/Animation/ContextualAnimation/Source/ContextualAnimation/Private/ContextualAnimSceneActorComponent.cpp
@@ -661,6 +661,11 @@ bool UContextualAnimSceneActorComponent::TransitionContextualAnimScene(FName Sec
 {
 	return TransitionContextualAnimScene(SectionName, {});
 }
+
+bool UContextualAnimSceneActorComponent::TransitionContextualAnimSceneToSpecificSet(FName SectionName, int32 AnimSetIdx)
+{
+	return TransitionContextualAnimScene(SectionName, AnimSetIdx, {});
+}
 	
 bool UContextualAnimSceneActorComponent::TransitionSingleActor(int32 SectionIdx, int32 AnimSetIdx)
 {
@@ -1068,15 +1073,21 @@ void UContextualAnimSceneActorComponent::SetMovementState(const FContextualAnimS
 	if (UCharacterMovementComponent* MovementComp = Binding.GetCharacterMovementComponent())
 	{
 		// Save movement state before the interaction starts so we can restore it when it ends
-		CharacterPropertiesBackup.bIgnoreClientMovementErrorChecksAndCorrection = MovementComp->bIgnoreClientMovementErrorChecksAndCorrection;
 		CharacterPropertiesBackup.bAllowPhysicsRotationDuringAnimRootMotion = MovementComp->bAllowPhysicsRotationDuringAnimRootMotion;
 		CharacterPropertiesBackup.bUseControllerDesiredRotation = MovementComp->bUseControllerDesiredRotation;
 		CharacterPropertiesBackup.bOrientRotationToMovement = MovementComp->bOrientRotationToMovement;
 		CharacterPropertiesBackup.MovementMode = MovementComp->MovementMode;
 		CharacterPropertiesBackup.bSimulatePhysics = MovementComp->UpdatedPrimitive && MovementComp->UpdatedPrimitive->IsSimulatingPhysics();
 
-		// Disable movement correction.
-		MovementComp->bIgnoreClientMovementErrorChecksAndCorrection = true;
+		// Disable movement correction if needed
+		if (const UContextualAnimSceneAsset* Asset = Bindings.GetSceneAsset())
+		{
+			if (Asset->ShouldIgnoreClientMovementErrorChecksAndCorrection())
+			{
+				CharacterPropertiesBackup.bIgnoreClientMovementErrorChecksAndCorrection = MovementComp->bIgnoreClientMovementErrorChecksAndCorrection;
+				MovementComp->bIgnoreClientMovementErrorChecksAndCorrection = true;
+			}
+		}
 
 		// Prevent physics rotation. During the interaction we want to be fully root motion driven
 		MovementComp->bAllowPhysicsRotationDuringAnimRootMotion = false;
@@ -1097,10 +1108,18 @@ void UContextualAnimSceneActorComponent::SetMovementState(const FContextualAnimS
 
 void UContextualAnimSceneActorComponent::RestoreMovementState(const FContextualAnimSceneBinding& Binding)
 {
+	// Restore movement state
 	if (UCharacterMovementComponent* MovementComp = Binding.GetCharacterMovementComponent())
 	{
-		// Restore movement state
-		MovementComp->bIgnoreClientMovementErrorChecksAndCorrection = CharacterPropertiesBackup.bIgnoreClientMovementErrorChecksAndCorrection;
+		// Restore movement correction if needed
+		if (const UContextualAnimSceneAsset* Asset = Bindings.GetSceneAsset())
+		{
+			if (Asset->ShouldIgnoreClientMovementErrorChecksAndCorrection())
+			{
+				MovementComp->bIgnoreClientMovementErrorChecksAndCorrection = CharacterPropertiesBackup.bIgnoreClientMovementErrorChecksAndCorrection;
+			}
+		}
+
 		MovementComp->bAllowPhysicsRotationDuringAnimRootMotion = CharacterPropertiesBackup.bAllowPhysicsRotationDuringAnimRootMotion;
 		MovementComp->bUseControllerDesiredRotation = CharacterPropertiesBackup.bUseControllerDesiredRotation;
 		MovementComp->bOrientRotationToMovement = CharacterPropertiesBackup.bOrientRotationToMovement;
diff --git a/Engine/Plugins/Experimental/Animation/ContextualAnimation/Source/ContextualAnimation/Public/ContextualAnimSceneActorComponent.h b/Engine/Plugins/Experimental/Animation/ContextualAnimation/Source/ContextualAnimation/Public/ContextualAnimSceneActorComponent.h
--- a/Engine/Plugins/Experimental/Animation/ContextualAnimation/Source/ContextualAnimation/Public/ContextualAnimSceneActorComponent.h
+++ b/Engine/Plugins/Experimental/Animation/ContextualAnimation/Source/ContextualAnimation/Public/ContextualAnimSceneActorComponent.h
@@ -205,6 +205,9 @@ public:
 	UFUNCTION(BlueprintCallable, Category = "Contextual Anim|Scene Actor Component")
 	bool TransitionContextualAnimScene(FName SectionName);
 	
+	UFUNCTION(BlueprintCallable, Category = "Contextual Anim|Scene Actor Component")
+	bool TransitionContextualAnimSceneToSpecificSet(FName SectionName, int32 AnimSetIdx);
+
 	UFUNCTION(BlueprintCallable, Category = "Contextual Anim|Scene Actor Component")
 	bool TransitionSingleActor(int32 SectionIdx, int32 AnimSetIdx);
 
diff --git a/Engine/Plugins/Experimental/Animation/ContextualAnimation/Source/ContextualAnimation/Public/ContextualAnimSceneAsset.h b/Engine/Plugins/Experimental/Animation/ContextualAnimation/Source/ContextualAnimation/Public/ContextualAnimSceneAsset.h
--- a/Engine/Plugins/Experimental/Animation/ContextualAnimation/Source/ContextualAnimation/Public/ContextualAnimSceneAsset.h
+++ b/Engine/Plugins/Experimental/Animation/ContextualAnimation/Source/ContextualAnimation/Public/ContextualAnimSceneAsset.h
@@ -238,6 +238,7 @@ public:
 	FORCEINLINE int32 GetSampleRate() const { return SampleRate; }
 	FORCEINLINE float GetRadius() const { return Radius; }
 	FORCEINLINE bool ShouldPrecomputeAlignmentTracks() const { return bPrecomputeAlignmentTracks; }
+	FORCEINLINE bool ShouldIgnoreClientMovementErrorChecksAndCorrection() const { return bIgnoreClientMovementErrorChecksAndCorrection; }
 
 	const TArray<TEnumAsByte<ECollisionChannel>>& GetCollisionChannelsToIgnoreForRole(FName Role) const;
 	
@@ -352,6 +353,9 @@ protected:
 	UPROPERTY(EditAnywhere, Category = "Settings")
 	TArray<FContextualAnimAttachmentParams> AttachmentParams;
 
+	UPROPERTY(EditAnywhere, Category = "Settings", AdvancedDisplay)
+	bool bIgnoreClientMovementErrorChecksAndCorrection = true;
+
 	/** Whether we should extract and cache alignment tracks off line. */
 	UPROPERTY(EditAnywhere, Category = "Settings", AdvancedDisplay)
 	bool bPrecomputeAlignmentTracks = false;
