Contextual Anim Editor - Disabled add anim set button on the toolbar if roles asset is not set + replaced check with ensure when attempting to construct bindings with an empty asset

diff --git a/Engine/Plugins/Experimental/Animation/ContextualAnimation/Source/ContextualAnimation/Private/ContextualAnimTypes.cpp b/Engine/Plugins/Experimental/Animation/ContextualAnimation/Source/ContextualAnimation/Private/ContextualAnimTypes.cpp
--- a/Engine/Plugins/Experimental/Animation/ContextualAnimation/Source/ContextualAnimation/Private/ContextualAnimTypes.cpp
+++ b/Engine/Plugins/Experimental/Animation/ContextualAnimation/Source/ContextualAnimation/Private/ContextualAnimTypes.cpp
@@ -360,11 +360,13 @@ int32 FContextualAnimSceneBinding::GetCurrentSectionIndex() const
 
 FContextualAnimSceneBindings::FContextualAnimSceneBindings(const UContextualAnimSceneAsset& InSceneAsset, int32 InSectionIdx, int32 InAnimSetIdx)
 {
-	check(InSceneAsset.HasValidData());
-	SceneAsset = &InSceneAsset;
-	SectionIdx = InSectionIdx;
-	AnimSetIdx = InAnimSetIdx;
-	GenerateUniqueId();
+	if (ensureAlways(InSceneAsset.HasValidData()))
+	{
+		SceneAsset = &InSceneAsset;
+		SectionIdx = InSectionIdx;
+		AnimSetIdx = InAnimSetIdx;
+		GenerateUniqueId();
+	}
 }
 
 void FContextualAnimSceneBindings::GenerateUniqueId()
diff --git a/Engine/Plugins/Experimental/Animation/ContextualAnimation/Source/ContextualAnimationEditor/Private/ContextualAnimAssetEditorToolkit.cpp b/Engine/Plugins/Experimental/Animation/ContextualAnimation/Source/ContextualAnimationEditor/Private/ContextualAnimAssetEditorToolkit.cpp
--- a/Engine/Plugins/Experimental/Animation/ContextualAnimation/Source/ContextualAnimationEditor/Private/ContextualAnimAssetEditorToolkit.cpp
+++ b/Engine/Plugins/Experimental/Animation/ContextualAnimation/Source/ContextualAnimationEditor/Private/ContextualAnimAssetEditorToolkit.cpp
@@ -14,6 +14,7 @@
 #include "ISequencer.h"
 #include "AdvancedPreviewSceneModule.h"
 #include "Modules/ModuleManager.h"
+#include "Misc/MessageDialog.h"
 
 #define LOCTEXT_NAMESPACE "ContextualAnimAssetEditorToolkit"
 
@@ -251,6 +252,16 @@ TSharedRef<SWidget> FContextualAnimAssetEditorToolkit::BuildSectionsMenu()
 
 void FContextualAnimAssetEditorToolkit::ShowNewAnimSetDialog()
 {
+	const UContextualAnimSceneAsset* SceneAsset = GetSceneAsset();
+	const int32 NumRoles = SceneAsset ? SceneAsset->GetNumRoles() : 0;
+
+	if (NumRoles == 0)
+	{
+		const FText DialogMsg = LOCTEXT("MissingRolesDialog", "Can't add new anim set. Roles asset hasn't been set or no roles has been defined.");
+		FMessageDialog::Open(EAppMsgType::Ok, DialogMsg);
+		return;
+	}
+
 	TSharedRef<SContextualAnimNewAnimSetDialog> NewAnimSetDialog = SNew(SContextualAnimNewAnimSetDialog, ViewModel.ToSharedRef());
 	NewAnimSetDialog->Show();
 }
diff --git a/Engine/Plugins/Experimental/Animation/ContextualAnimation/Source/ContextualAnimationEditor/Private/ContextualAnimViewModel.cpp b/Engine/Plugins/Experimental/Animation/ContextualAnimation/Source/ContextualAnimationEditor/Private/ContextualAnimViewModel.cpp
--- a/Engine/Plugins/Experimental/Animation/ContextualAnimation/Source/ContextualAnimationEditor/Private/ContextualAnimViewModel.cpp
+++ b/Engine/Plugins/Experimental/Animation/ContextualAnimation/Source/ContextualAnimationEditor/Private/ContextualAnimViewModel.cpp
@@ -549,6 +549,11 @@ void FContextualAnimViewModel::RefreshPreviewScene()
 		}
 	}
 
+	if (!SceneAsset->HasValidData())
+	{
+		return;
+	}
+
 	if (!SceneAsset->Sections.IsValidIndex(ActiveSectionIdx))
 	{
 		return;
