Contextual Anim - Stop listening to montage callbacks right before EndPlay so we don't trigger any of the logic for when an interaction ends once the game ends. This prevents a few ensures from triggering at lower levels when trying to change the movement mode during the unregister phase of the components.

diff --git a/Engine/Plugins/Experimental/Animation/ContextualAnimation/Source/ContextualAnimation/Private/ContextualAnimSceneActorComponent.cpp b/Engine/Plugins/Experimental/Animation/ContextualAnimation/Source/ContextualAnimation/Private/ContextualAnimSceneActorComponent.cpp
--- a/Engine/Plugins/Experimental/Animation/ContextualAnimation/Source/ContextualAnimation/Private/ContextualAnimSceneActorComponent.cpp
+++ b/Engine/Plugins/Experimental/Animation/ContextualAnimation/Source/ContextualAnimation/Private/ContextualAnimSceneActorComponent.cpp
@@ -100,6 +100,17 @@ void UContextualAnimSceneActorComponent::GetLifetimeReplicatedProps(TArray< FLif
 	DOREPLIFETIME_WITH_PARAMS_FAST(UContextualAnimSceneActorComponent, RepTransitionData, Params);
 }
 
+void UContextualAnimSceneActorComponent::EndPlay(const EEndPlayReason::Type EndPlayReason)
+{
+	if (OwnerAnimInstance.IsValid())
+	{
+		OwnerAnimInstance->OnMontageBlendingOut.RemoveDynamic(this, &UContextualAnimSceneActorComponent::OnMontageBlendingOut);
+		OwnerAnimInstance->OnPlayMontageNotifyBegin.RemoveDynamic(this, &UContextualAnimSceneActorComponent::OnPlayMontageNotifyBegin);
+	}
+
+	Super::EndPlay(EndPlayReason);
+}
+
 bool UContextualAnimSceneActorComponent::IsOwnerLocallyControlled() const
 {
 	if (const APawn* OwnerPawn = Cast<APawn>(GetOwner()))
@@ -122,6 +133,9 @@ void UContextualAnimSceneActorComponent::PlayAnimation_Internal(UAnimSequenceBas
 		UE_LOG(LogContextualAnim, Verbose, TEXT("%-21s \t\tUContextualAnimSceneActorComponent::PlayAnimation_Internal Playing Animation. Actor: %s Anim: %s StartTime: %f bSyncPlaybackTime: %d"),
 			*UEnum::GetValueAsString(TEXT("Engine.ENetRole"), GetOwner()->GetLocalRole()), *GetNameSafe(GetOwner()), *GetNameSafe(Animation), StartTime, bSyncPlaybackTime);
 
+		// Cache AnimInstance so we don't have to look for it in the bindings in subsequent uses.
+		OwnerAnimInstance = AnimInstance;
+
 		//@TODO: Add support for dynamic montage
 		UAnimMontage* AnimMontage = Cast<UAnimMontage>(Animation);
 
@@ -132,7 +146,6 @@ void UContextualAnimSceneActorComponent::PlayAnimation_Internal(UAnimSequenceBas
 
 		AnimInstance->OnMontageBlendingOut.AddUniqueDynamic(this, &UContextualAnimSceneActorComponent::OnMontageBlendingOut);
 		AnimInstance->OnPlayMontageNotifyBegin.AddUniqueDynamic(this, &UContextualAnimSceneActorComponent::OnPlayMontageNotifyBegin);
-		
 
 		if (bSyncPlaybackTime)
 		{
@@ -987,18 +1000,20 @@ void UContextualAnimSceneActorComponent::LeaveScene()
 			*UEnum::GetValueAsString(TEXT("Engine.ENetRole"), GetOwner()->GetLocalRole()), *GetNameSafe(GetOwner()), *Bindings.GetRoleFromBinding(*Binding).ToString(),
 			Bindings.GetID(), Bindings.GetSectionIdx(), *GetNameSafe(Bindings.GetSceneAsset()));
 
-		if (UAnimInstance* AnimInstance = Binding->GetAnimInstance())
+		if (OwnerAnimInstance.IsValid())
 		{
-			AnimInstance->OnMontageBlendingOut.RemoveDynamic(this, &UContextualAnimSceneActorComponent::OnMontageBlendingOut);
-			AnimInstance->OnPlayMontageNotifyBegin.RemoveDynamic(this, &UContextualAnimSceneActorComponent::OnPlayMontageNotifyBegin);
+			OwnerAnimInstance->OnMontageBlendingOut.RemoveDynamic(this, &UContextualAnimSceneActorComponent::OnMontageBlendingOut);
+			OwnerAnimInstance->OnPlayMontageNotifyBegin.RemoveDynamic(this, &UContextualAnimSceneActorComponent::OnPlayMontageNotifyBegin);
 
 			//@TODO: Add support for dynamic montage
-			const UAnimMontage* AnimMontage = AnimInstance->GetCurrentActiveMontage();
+			const UAnimMontage* AnimMontage = OwnerAnimInstance->GetCurrentActiveMontage();
 			if (AnimMontage)
 			{
 				UE_LOG(LogContextualAnim, VeryVerbose, TEXT("\t\t Stopping animation (%s) from LeaveScene"), *GetNameSafe(AnimMontage));
-				AnimInstance->Montage_Stop(AnimMontage->GetDefaultBlendOutTime());
+				OwnerAnimInstance->Montage_Stop(AnimMontage->GetDefaultBlendOutTime());
 			}
+
+			OwnerAnimInstance.Reset();
 		}
 
 		// Stop listening to TickPose if we were
diff --git a/Engine/Plugins/Experimental/Animation/ContextualAnimation/Source/ContextualAnimation/Public/ContextualAnimSceneActorComponent.h b/Engine/Plugins/Experimental/Animation/ContextualAnimation/Source/ContextualAnimation/Public/ContextualAnimSceneActorComponent.h
--- a/Engine/Plugins/Experimental/Animation/ContextualAnimation/Source/ContextualAnimation/Public/ContextualAnimSceneActorComponent.h
+++ b/Engine/Plugins/Experimental/Animation/ContextualAnimation/Source/ContextualAnimation/Public/ContextualAnimSceneActorComponent.h
@@ -171,6 +171,7 @@ public:
 
 	static void AddReferencedObjects(UObject* InThis, FReferenceCollector& Collector);
 	virtual void GetLifetimeReplicatedProps(TArray< FLifetimeProperty >& OutLifetimeProps) const;
+	virtual void EndPlay(const EEndPlayReason::Type EndPlayReason) override;
 	virtual FBoxSphereBounds CalcBounds(const FTransform& LocalToWorld) const;
 	virtual FPrimitiveSceneProxy* CreateSceneProxy() override;
 
@@ -225,6 +226,10 @@ public:
 
 protected:
 
+	/** Cache of the owner's AnimInstance for easy access. Only valid while an interaction is active */
+	UPROPERTY(Transient)
+	TWeakObjectPtr<UAnimInstance> OwnerAnimInstance = nullptr;
+
 	/** 
 	 * Replicated copy of the bindings so we can start the action on simulated proxies 
 	 * This gets replicated only from the initiator of the action and then set on all the other members of the interaction
