Fix for UContextualAnimSceneActorComponent not exiting a contextual animation if a member of that context has no animation to play.

diff --git a/Engine/Plugins/Experimental/Animation/ContextualAnimation/Source/ContextualAnimation/Private/ContextualAnimSceneActorComponent.cpp b/Engine/Plugins/Experimental/Animation/ContextualAnimation/Source/ContextualAnimation/Private/ContextualAnimSceneActorComponent.cpp
--- a/Engine/Plugins/Experimental/Animation/ContextualAnimation/Source/ContextualAnimation/Private/ContextualAnimSceneActorComponent.cpp
+++ b/Engine/Plugins/Experimental/Animation/ContextualAnimation/Source/ContextualAnimation/Private/ContextualAnimSceneActorComponent.cpp
@@ -751,43 +751,40 @@ void UContextualAnimSceneActorComponent::EarlyOutContextualAnimScene()
 	{
 		const UAnimInstance* AnimInstance = Binding->GetAnimInstance();
 		const UAnimMontage* ActiveMontage = AnimInstance ? AnimInstance->GetCurrentActiveMontage() : nullptr;
-		if (ActiveMontage)
+		UE_LOG(LogContextualAnim, Verbose, TEXT("%-21s UContextualAnimSceneActorComponent::EarlyOutContextualAnimScene Actor: %s ActiveMontage: %s"),
+			*UEnum::GetValueAsString(TEXT("Engine.ENetRole"), GetOwner()->GetLocalRole()), *GetNameSafe(GetOwner()), *GetNameSafe(ActiveMontage));
+
+		if (Bindings.GetAnimTrackFromBinding(*Binding).Animation == ActiveMontage)
 		{
-			UE_LOG(LogContextualAnim, Verbose, TEXT("%-21s UContextualAnimSceneActorComponent::EarlyOutContextualAnimScene Actor: %s ActiveMontage: %s"),
-				*UEnum::GetValueAsString(TEXT("Engine.ENetRole"), GetOwner()->GetLocalRole()), *GetNameSafe(GetOwner()), *GetNameSafe(ActiveMontage));
+			const uint8 BindingsId = Bindings.GetID();
+
+			// Stop animation.
+			LeaveScene();
 
-			if (Bindings.GetAnimTrackFromBinding(*Binding).Animation == ActiveMontage)
+			// If we are on the server, rep the event to stop animation on simulated proxies
+			if (GetOwner()->HasAuthority())
 			{
-				const uint8 BindingsId = Bindings.GetID();
+				RepTransitionSingleActorData.Id = BindingsId;
+				RepTransitionSingleActorData.SectionIdx = MAX_uint8;
+				RepTransitionSingleActorData.AnimSetIdx = MAX_uint8;
+				RepTransitionSingleActorData.ExternalWarpTargets.Reset();
+				RepTransitionSingleActorData.IncrementRepCounter();
 
-				// Stop animation.
-				LeaveScene();
+				RepLateJoinData.Reset();
+				RepTransitionData.Reset();
+				RepBindings.Reset();
 
-				// If we are on the server, rep the event to stop animation on simulated proxies
-				if (GetOwner()->HasAuthority())
-				{
-					RepTransitionSingleActorData.Id = BindingsId;
-					RepTransitionSingleActorData.SectionIdx = MAX_uint8;
-					RepTransitionSingleActorData.AnimSetIdx = MAX_uint8;
-					RepTransitionSingleActorData.ExternalWarpTargets.Reset();
-					RepTransitionSingleActorData.IncrementRepCounter();
-
-					RepLateJoinData.Reset();
-					RepTransitionData.Reset();
-					RepBindings.Reset();
-
-					MARK_PROPERTY_DIRTY_FROM_NAME(UContextualAnimSceneActorComponent, RepTransitionSingleActorData, this);
-					MARK_PROPERTY_DIRTY_FROM_NAME(UContextualAnimSceneActorComponent, RepLateJoinData, this);
-					MARK_PROPERTY_DIRTY_FROM_NAME(UContextualAnimSceneActorComponent, RepTransitionData, this);
-					MARK_PROPERTY_DIRTY_FROM_NAME(UContextualAnimSceneActorComponent, RepBindings, this);
-
-					GetOwner()->ForceNetUpdate();
-				}
-				// If local player, tell the server to stop the animation too
-				else if (GetOwner()->GetLocalRole() == ROLE_AutonomousProxy)
-				{
-					ServerEarlyOutContextualAnimScene();
-				}
+				MARK_PROPERTY_DIRTY_FROM_NAME(UContextualAnimSceneActorComponent, RepTransitionSingleActorData, this);
+				MARK_PROPERTY_DIRTY_FROM_NAME(UContextualAnimSceneActorComponent, RepLateJoinData, this);
+				MARK_PROPERTY_DIRTY_FROM_NAME(UContextualAnimSceneActorComponent, RepTransitionData, this);
+				MARK_PROPERTY_DIRTY_FROM_NAME(UContextualAnimSceneActorComponent, RepBindings, this);
+
+				GetOwner()->ForceNetUpdate();
+			}
+			// If local player, tell the server to stop the animation too
+			else if (GetOwner()->GetLocalRole() == ROLE_AutonomousProxy)
+			{
+				ServerEarlyOutContextualAnimScene();
 			}
 		}
 	}
