Contextual Anim - IK Targets improvements
 - Ik Targets are now defined at the asset level instead of per section
 - Static actors can now provide ik targets (before it would only work for skeletal meshes)
 - Added option to define the alpha value for ik targets from a curve in the animation
 - Added blueprint callable function to get the primary binding from bindings

diff --git a/Engine/Plugins/Experimental/Animation/ContextualAnimation/Source/ContextualAnimation/Private/ContextualAnimSceneActorComponent.cpp b/Engine/Plugins/Experimental/Animation/ContextualAnimation/Source/ContextualAnimation/Private/ContextualAnimSceneActorComponent.cpp
--- a/Engine/Plugins/Experimental/Animation/ContextualAnimation/Source/ContextualAnimation/Private/ContextualAnimSceneActorComponent.cpp
+++ b/Engine/Plugins/Experimental/Animation/ContextualAnimation/Source/ContextualAnimation/Private/ContextualAnimSceneActorComponent.cpp
@@ -339,6 +339,8 @@ void UContextualAnimSceneActorComponent::LateJoinScene(const FContextualAnimScen
 		SetMovementState(*Binding, AnimTrack->MovementMode, AnimTrack->CustomMovementMode);
 
 		OnLateJoinScene(*Binding, SectionIdx, AnimSetIdx);
+
+		OnJoinedSceneDelegate.Broadcast(this);
 	}
 }
 
@@ -1034,6 +1036,8 @@ void UContextualAnimSceneActorComponent::LeaveScene()
 			WarpTargetNamesCache.Reset();
 		}
 
+		IKTargets.Reset();
+
 		// Notify the other actors in the interaction
 		// @TODO: This should be refactored so only the leader of the interaction maintains the full bindings
 		for (const FContextualAnimSceneBinding& OtherBinding : Bindings)
@@ -1212,14 +1216,24 @@ void UContextualAnimSceneActorComponent::OnPlayMontageNotifyBegin(FName NotifyNa
 
 void UContextualAnimSceneActorComponent::OnTickPose(class USkinnedMeshComponent* SkinnedMeshComponent, float DeltaTime, bool bNeedsValidRootMotion)
 {
-	//@TODO: Check for LOD to prevent this update if the actor is too far away
-	UpdateIKTargets();
+	//@TODO: Check for LOD too to prevent this update if the actor is too far away
+	if (GetNetMode() != NM_DedicatedServer)
+	{
+		UpdateIKTargets();
+	}
 }
 
 void UContextualAnimSceneActorComponent::UpdateIKTargets()
 {
+	QUICK_SCOPE_CYCLE_COUNTER(STAT_ContextualAnim_UpdateIKTargets);
+
 	IKTargets.Reset();
 
+	if (!Bindings.IsValid())
+	{
+		return;
+	}
+
 	const FContextualAnimSceneBinding* BindingPtr = Bindings.FindBindingByActor(GetOwner());
 	if (BindingPtr == nullptr)
 	{
@@ -1233,9 +1247,31 @@ void UContextualAnimSceneActorComponent::UpdateIKTargets()
 	}
 
 	const TArray<FContextualAnimIKTargetDefinition>& IKTargetDefs = Bindings.GetIKTargetDefContainerFromBinding(*BindingPtr).IKTargetDefs;
+
 	for (const FContextualAnimIKTargetDefinition& IKTargetDef : IKTargetDefs)
 	{
-		float Alpha = UAnimNotifyState_IKWindow::GetIKAlphaValue(IKTargetDef.GoalName, MontageInstance);
+		float Alpha = 0.f;
+
+		const EContextualAnimIKTargetAlphaProvider AlphaProvider = Bindings.GetSceneAsset()->GetIKTargetParams().AlphaProvider;
+		if (AlphaProvider == EContextualAnimIKTargetAlphaProvider::AnimNotifyState)
+		{
+			Alpha = UAnimNotifyState_IKWindow::GetIKAlphaValue(IKTargetDef.GoalName, MontageInstance);
+		}
+		else if (AlphaProvider == EContextualAnimIKTargetAlphaProvider::Curve)
+		{
+			if (const UAnimInstance* AnimInstance = BindingPtr->GetAnimInstance())
+			{
+				Alpha = AnimInstance->GetCurveValue(IKTargetDef.GoalName);
+			}
+		}
+		else if (AlphaProvider == EContextualAnimIKTargetAlphaProvider::None)
+		{
+			Alpha = 1.f;
+		}
+		else
+		{
+			check(false); // Unreachable
+		}
 
 		// @TODO: IKTargetTransform will be off by 1 frame if we tick before target. 
 		// Should we at least add an option to the SceneAsset to setup tick dependencies or should this be entirely up to the user?
@@ -1257,39 +1293,49 @@ void UContextualAnimSceneActorComponent::UpdateIKTargets()
 
 			if (Alpha > 0.f)
 			{
-				if (const USkeletalMeshComponent* TargetSkelMeshComp = TargetBinding->GetSkeletalMeshComponent())
+				if (const UMeshComponent* TargetMeshComp = UContextualAnimUtilities::TryGetMeshComponentWithSocket(TargetBinding->GetActor(), IKTargetDef.TargetBoneName))
 				{
 					if (IKTargetDef.Provider == EContextualAnimIKTargetProvider::Autogenerated)
 					{
-						const FTransform IKTargetParentTransform = TargetSkelMeshComp->GetSocketTransform(IKTargetDef.TargetBoneName);
+						const FTransform IKTargetParentTransform = TargetMeshComp->GetSocketTransform(IKTargetDef.TargetBoneName);
 
 						const float Time = MontageInstance->GetPosition();
 						const FTransform IKTargetTransform = Bindings.GetIKTargetTransformFromBinding(*BindingPtr, IKTargetDef.GoalName, Time) * IKTargetParentTransform;
 
-						IKTargets.Add(FContextualAnimIKTarget(IKTargetDef.GoalName, Alpha, IKTargetTransform));
+						IKTargets.Add(FContextualAnimIKTarget(IKTargetDef.GoalName, IKTargetDef.BoneName, Alpha, IKTargetTransform));
 
 #if !(UE_BUILD_SHIPPING || UE_BUILD_TEST)
 						if (CVarContextualAnimIKDebug.GetValueOnGameThread() > 0)
 						{
+							const FColor Color = FColor::MakeRedToGreenColorFromScalar(Alpha);
 							const float DrawDebugDuration = CVarContextualAnimIKDrawDebugLifetime.GetValueOnGameThread();
-							DrawDebugLine(GetWorld(), IKTargetParentTransform.GetLocation(), IKTargetTransform.GetLocation(), FColor::MakeRedToGreenColorFromScalar(Alpha), false, DrawDebugDuration, 0, 0.5f);
+							DrawDebugLine(GetWorld(), IKTargetParentTransform.GetLocation(), IKTargetTransform.GetLocation(), Color, false, DrawDebugDuration, 0, 0.5f);
 							DrawDebugCoordinateSystem(GetWorld(), IKTargetTransform.GetLocation(), IKTargetTransform.Rotator(), 10.f, false, DrawDebugDuration, 0, 0.5f);
+
+							if (const USkeletalMeshComponent* SkelMesh = BindingPtr->GetSkeletalMeshComponent())
+							{
+								DrawDebugDirectionalArrow(GetWorld(), SkelMesh->GetSocketLocation(IKTargetDef.BoneName), IKTargetTransform.GetLocation(), 5.f, Color, false, DrawDebugDuration, 0, 0.5f);
+							}
 						}
 #endif
 					}
 					else if (IKTargetDef.Provider == EContextualAnimIKTargetProvider::Bone)
 					{
-						const FTransform IKTargetTransform = TargetSkelMeshComp->GetSocketTransform(IKTargetDef.TargetBoneName);
+						const FTransform IKTargetTransform = TargetMeshComp->GetSocketTransform(IKTargetDef.TargetBoneName);
 
-						IKTargets.Add(FContextualAnimIKTarget(IKTargetDef.GoalName, Alpha, IKTargetTransform));
+						IKTargets.Add(FContextualAnimIKTarget(IKTargetDef.GoalName, IKTargetDef.BoneName, Alpha, IKTargetTransform));
 
 #if !(UE_BUILD_SHIPPING || UE_BUILD_TEST)
 						if (CVarContextualAnimIKDebug.GetValueOnGameThread() > 0)
 						{
+							const FColor Color = FColor::MakeRedToGreenColorFromScalar(Alpha);
 							const float DrawDebugDuration = CVarContextualAnimIKDrawDebugLifetime.GetValueOnGameThread();
-							const FTransform IKTargetParentTransform = TargetSkelMeshComp->GetSocketTransform(TargetSkelMeshComp->GetParentBone(IKTargetDef.TargetBoneName));
-							DrawDebugLine(GetWorld(), IKTargetParentTransform.GetLocation(), IKTargetTransform.GetLocation(), FColor::MakeRedToGreenColorFromScalar(Alpha), false, DrawDebugDuration, 0, 0.5f);
 							DrawDebugCoordinateSystem(GetWorld(), IKTargetTransform.GetLocation(), IKTargetTransform.Rotator(), 10.f, false, DrawDebugDuration, 0, 0.5f);
+
+							if (const USkeletalMeshComponent* SkelMesh = BindingPtr->GetSkeletalMeshComponent())
+							{
+								DrawDebugDirectionalArrow(GetWorld(), SkelMesh->GetSocketLocation(IKTargetDef.BoneName), IKTargetTransform.GetLocation(), 5.f, Color, false, DrawDebugDuration, 0, 0.5f);
+							}
 						}
 #endif
 					}
diff --git a/Engine/Plugins/Experimental/Animation/ContextualAnimation/Source/ContextualAnimation/Private/ContextualAnimSceneAsset.cpp b/Engine/Plugins/Experimental/Animation/ContextualAnimation/Source/ContextualAnimation/Private/ContextualAnimSceneAsset.cpp
--- a/Engine/Plugins/Experimental/Animation/ContextualAnimation/Source/ContextualAnimation/Private/ContextualAnimSceneAsset.cpp
+++ b/Engine/Plugins/Experimental/Animation/ContextualAnimation/Source/ContextualAnimation/Private/ContextualAnimSceneAsset.cpp
@@ -171,8 +171,8 @@ void FContextualAnimSceneSection::GenerateIKTargetTracks(UContextualAnimSceneAss
 		{
 			AnimTrack.IKTargetData.Empty();
 
-			const FContextualAnimIKTargetDefContainer* IKTargetDefContainerPtr = RoleToIKTargetDefsMap.Find(AnimTrack.Role);
-			if (IKTargetDefContainerPtr == nullptr || IKTargetDefContainerPtr->IKTargetDefs.Num() == 0)
+			const TArray<FContextualAnimIKTargetDefinition>& IKTargetDefs = SceneAsset.GetIKTargetDefsForRole(AnimTrack.Role).IKTargetDefs;
+			if (IKTargetDefs.Num() == 0)
 			{
 				continue;
 			}
@@ -192,10 +192,10 @@ void FContextualAnimSceneSection::GenerateIKTargetTracks(UContextualAnimSceneAss
 					TArray<TTuple<FName, FName, int32, FName, int32>> BonesData; //0: GoalName, 1: MyBoneName 2: MyBoneIndex, 3: TargetBoneName, 4: TargetBoneIndex
 				};
 				TMap<FName, FPoseExtractionHelper> PoseExtractionHelperMap;
-				PoseExtractionHelperMap.Reserve(IKTargetDefContainerPtr->IKTargetDefs.Num());
+				PoseExtractionHelperMap.Reserve(IKTargetDefs.Num());
 
 				int32 TotalTracks = 0;
-				for (const FContextualAnimIKTargetDefinition& IKTargetDef : IKTargetDefContainerPtr->IKTargetDefs)
+				for (const FContextualAnimIKTargetDefinition& IKTargetDef : IKTargetDefs)
 				{
 					if (IKTargetDef.Provider == EContextualAnimIKTargetProvider::Autogenerated)
 					{
@@ -376,16 +376,6 @@ FTransform FContextualAnimSceneSection::GetIKTargetTransformForRoleAtTime(int32
 	return FTransform::Identity;
 }
 
-const FContextualAnimIKTargetDefContainer& FContextualAnimSceneSection::GetIKTargetDefsForRole(const FName& Role) const
-{
-	if (const FContextualAnimIKTargetDefContainer* ContainerPtr = RoleToIKTargetDefsMap.Find(Role))
-	{
-		return *ContainerPtr;
-	}
-
-	return FContextualAnimIKTargetDefContainer::EmptyContainer;
-}
-
 const FContextualAnimTrack* FContextualAnimSceneSection::FindFirstAnimTrackForRoleThatPassesSelectionCriteria(const FName& Role, const FContextualAnimSceneBindingContext& Primary, const FContextualAnimSceneBindingContext& Querier) const
 {
 	for (const FContextualAnimSet& AnimSet : AnimSets)
@@ -652,9 +642,10 @@ int32 UContextualAnimSceneAsset::GetNumAnimSetsInSection(int32 SectionIdx) const
 	return Sections.IsValidIndex(SectionIdx) ? Sections[SectionIdx].AnimSets.Num() : 0;
 }
 
-const FContextualAnimIKTargetDefContainer& UContextualAnimSceneAsset::GetIKTargetDefsForRoleInSection(int32 SectionIdx, const FName& Role) const
+const FContextualAnimIKTargetDefContainer& UContextualAnimSceneAsset::GetIKTargetDefsForRole(const FName& Role) const
 {
-	return Sections.IsValidIndex(SectionIdx) ? Sections[SectionIdx].GetIKTargetDefsForRole(Role) : FContextualAnimIKTargetDefContainer::EmptyContainer;
+	const FContextualAnimIKTargetDefContainer* Ptr = IKTargetParams.IKTargetDefsForEachRole.FindByPredicate([Role](const FContextualAnimIKTargetDefContainer& Item) { return Item.Role == Role; });
+	return Ptr ? *Ptr : FContextualAnimIKTargetDefContainer::EmptyContainer;
 }
 
 FTransform UContextualAnimSceneAsset::GetIKTargetTransform(int32 SectionIdx, int32 AnimSetIdx, int32 AnimTrackIdx, const FName& TrackName, float Time) const
diff --git a/Engine/Plugins/Experimental/Animation/ContextualAnimation/Source/ContextualAnimation/Private/ContextualAnimTypes.cpp b/Engine/Plugins/Experimental/Animation/ContextualAnimation/Source/ContextualAnimation/Private/ContextualAnimTypes.cpp
--- a/Engine/Plugins/Experimental/Animation/ContextualAnimation/Source/ContextualAnimation/Private/ContextualAnimTypes.cpp
+++ b/Engine/Plugins/Experimental/Animation/ContextualAnimation/Source/ContextualAnimation/Private/ContextualAnimTypes.cpp
@@ -484,7 +484,7 @@ const FContextualAnimIKTargetDefContainer& FContextualAnimSceneBindings::GetIKTa
 		return FContextualAnimIKTargetDefContainer::EmptyContainer;
 	}
 
-	return GetSceneAsset()->GetIKTargetDefsForRoleInSection(GetSectionIdx(), GetRoleFromBinding(Binding));
+	return GetSceneAsset()->GetIKTargetDefsForRole(GetRoleFromBinding(Binding));
 }
 
 FTransform FContextualAnimSceneBindings::GetIKTargetTransformFromBinding(const FContextualAnimSceneBinding& Binding, const FName& TrackName, float Time) const
diff --git a/Engine/Plugins/Experimental/Animation/ContextualAnimation/Source/ContextualAnimation/Private/ContextualAnimUtilities.cpp b/Engine/Plugins/Experimental/Animation/ContextualAnimation/Source/ContextualAnimation/Private/ContextualAnimUtilities.cpp
--- a/Engine/Plugins/Experimental/Animation/ContextualAnimation/Source/ContextualAnimation/Private/ContextualAnimUtilities.cpp
+++ b/Engine/Plugins/Experimental/Animation/ContextualAnimation/Source/ContextualAnimation/Private/ContextualAnimUtilities.cpp
@@ -407,6 +407,16 @@ const FContextualAnimSceneBinding& UContextualAnimUtilities::BP_SceneBindings_Ge
 	return FContextualAnimSceneBinding::InvalidBinding;
 }
 
+const FContextualAnimSceneBinding& UContextualAnimUtilities::BP_SceneBindings_GetPrimaryBinding(const FContextualAnimSceneBindings& Bindings)
+{
+	if (const FContextualAnimSceneBinding* SceneActorData = Bindings.GetPrimaryBinding())
+	{
+		return *SceneActorData;
+	}
+
+	return FContextualAnimSceneBinding::InvalidBinding;
+}
+
 void UContextualAnimUtilities::BP_SceneBindings_CalculateWarpPoints(const FContextualAnimSceneBindings& Bindings, TArray<FContextualAnimWarpPoint>& OutWarpPoints)
 {
 	Bindings.CalculateWarpPoints(OutWarpPoints);
diff --git a/Engine/Plugins/Experimental/Animation/ContextualAnimation/Source/ContextualAnimation/Public/ContextualAnimSceneAsset.h b/Engine/Plugins/Experimental/Animation/ContextualAnimation/Source/ContextualAnimation/Public/ContextualAnimSceneAsset.h
--- a/Engine/Plugins/Experimental/Animation/ContextualAnimation/Source/ContextualAnimation/Public/ContextualAnimSceneAsset.h
+++ b/Engine/Plugins/Experimental/Animation/ContextualAnimation/Source/ContextualAnimation/Public/ContextualAnimSceneAsset.h
@@ -105,8 +105,6 @@ public:
 
 	FTransform GetIKTargetTransformForRoleAtTime(int32 AnimSetIdx, FName Role, FName TrackName, float Time) const;
 
-	const FContextualAnimIKTargetDefContainer& GetIKTargetDefsForRole(const FName& Role) const;
-
 	const FContextualAnimTrack* FindFirstAnimTrackForRoleThatPassesSelectionCriteria(const FName& Role, const FContextualAnimSceneBindingContext& Primary, const FContextualAnimSceneBindingContext& Querier) const;
 
 	FORCEINLINE FName GetName() const { return Name; }
@@ -122,9 +120,6 @@ protected:
 	UPROPERTY(EditAnywhere, Category = "Defaults")
 	TArray<FContextualAnimSet> AnimSets;
 
-	UPROPERTY(EditAnywhere, Category = "Defaults")
-	TMap<FName, FContextualAnimIKTargetDefContainer> RoleToIKTargetDefsMap;
-
 	UPROPERTY(EditAnywhere, Category = "Defaults", meta = (TitleProperty = "WarpTargetName"))
 	TArray<FContextualAnimWarpPointDefinition> WarpPointDefinitions;
 
@@ -247,6 +242,8 @@ public:
 		return AttachmentParams.FindByPredicate([Role](const FContextualAnimAttachmentParams& Item) { return Item.Role == Role; });
 	}
 
+	const FContextualAnimIKTargetParams& GetIKTargetParams() const { return IKTargetParams; }
+
 	bool HasValidData() const { return RolesAsset != nullptr && Sections.Num() > 0 && Sections[0].AnimSets.Num() > 0; }
 
 	const UContextualAnimRolesAsset* GetRolesAsset() const { return RolesAsset; }
@@ -290,7 +287,7 @@ public:
 
 	const FContextualAnimTrack* FindAnimTrackByAnimation(const UAnimSequenceBase* Animation) const;
 
-	const FContextualAnimIKTargetDefContainer& GetIKTargetDefsForRoleInSection(int32 SectionIdx, const FName& Role) const;
+	const FContextualAnimIKTargetDefContainer& GetIKTargetDefsForRole(const FName& Role) const;
 
 	UFUNCTION(BlueprintCallable, Category = "Contextual Anim|Scene Asset")
 	void GetAlignmentPointsForSecondaryRole(EContextualAnimPointType Type, int32 SectionIdx, const FContextualAnimSceneBindingContext& Primary, TArray<FContextualAnimPoint>& OutResult) const;
@@ -353,6 +350,9 @@ protected:
 	UPROPERTY(EditAnywhere, Category = "Settings")
 	TArray<FContextualAnimAttachmentParams> AttachmentParams;
 
+	UPROPERTY(EditAnywhere, Category = "Settings", meta = (DisplayName = "IK Target Params"))
+	FContextualAnimIKTargetParams IKTargetParams;
+
 	UPROPERTY(EditAnywhere, Category = "Settings", AdvancedDisplay)
 	bool bIgnoreClientMovementErrorChecksAndCorrection = true;
 
diff --git a/Engine/Plugins/Experimental/Animation/ContextualAnimation/Source/ContextualAnimation/Public/ContextualAnimTypes.h b/Engine/Plugins/Experimental/Animation/ContextualAnimation/Source/ContextualAnimation/Public/ContextualAnimTypes.h
--- a/Engine/Plugins/Experimental/Animation/ContextualAnimation/Source/ContextualAnimation/Public/ContextualAnimTypes.h
+++ b/Engine/Plugins/Experimental/Animation/ContextualAnimation/Source/ContextualAnimation/Public/ContextualAnimTypes.h
@@ -161,7 +161,7 @@ struct FContextualAnimIKTargetDefinition
 	UPROPERTY(EditAnywhere, BlueprintReadWrite, Category = "Defaults")
 	EContextualAnimIKTargetProvider Provider = EContextualAnimIKTargetProvider::Autogenerated;
 
-	UPROPERTY(EditAnywhere, BlueprintReadWrite, Category = "Defaults")
+	UPROPERTY(EditAnywhere, BlueprintReadWrite, Category = "Defaults", meta = (GetOptions = "GetRoles"))
 	FName TargetRoleName = NAME_None;
 
 	UPROPERTY(EditAnywhere, BlueprintReadWrite, Category = "Defaults")
@@ -185,6 +185,9 @@ struct CONTEXTUALANIMATION_API FContextualAnimIKTarget
 	UPROPERTY(EditAnywhere, BlueprintReadWrite, Category = "Defaults")
 	FName GoalName = NAME_None;
 
+	UPROPERTY(EditAnywhere, BlueprintReadWrite, Category = "Defaults")
+	FName BoneName = NAME_None;
+
 	UPROPERTY(EditAnywhere, BlueprintReadWrite, Category = "Defaults")
 	float Alpha = 0.f;
 
@@ -192,8 +195,8 @@ struct CONTEXTUALANIMATION_API FContextualAnimIKTarget
 	FTransform Transform;
 
 	FContextualAnimIKTarget() {}
-	FContextualAnimIKTarget(const FName& InGoalName, float InAlpha, const FTransform& InTransform)
-		: GoalName(InGoalName), Alpha(InAlpha), Transform(InTransform) {}
+	FContextualAnimIKTarget(const FName& InGoalName, const FName& InBoneName, float InAlpha, const FTransform& InTransform)
+		: GoalName(InGoalName), BoneName(InBoneName), Alpha(InAlpha), Transform(InTransform) {}
 
 	static const FContextualAnimIKTarget InvalidIKTarget;
 };
@@ -203,12 +206,50 @@ struct CONTEXTUALANIMATION_API FContextualAnimIKTargetDefContainer
 {
 	GENERATED_BODY()
 
-	UPROPERTY(EditAnywhere, BlueprintReadWrite, Category = "Defaults")
+	/** Role this IK Target definitions are for */
+	UPROPERTY(EditAnywhere, BlueprintReadWrite, Category = "Defaults", meta = (GetOptions = "GetRoles"))
+	FName Role = NAME_None;
+
+	/** List of IK Target definitions */
+	UPROPERTY(EditAnywhere, BlueprintReadWrite, Category = "Defaults", meta = (DisplayName = "Definitions"))
 	TArray<FContextualAnimIKTargetDefinition> IKTargetDefs;
 
 	static const FContextualAnimIKTargetDefContainer EmptyContainer;
 };
 
+UENUM(BlueprintType)
+enum class EContextualAnimIKTargetAlphaProvider : uint8
+{	
+	/** 
+	 * Alpha value is determine by an IK Window in the animation. 
+	 * IK Target transform will only be calculated when alpha is bigger than 0 
+	 */
+	AnimNotifyState,
+
+	/** 
+	 * Alpha value is determined by a curve in the animation with the same name as the IK Target
+	 * IK Target transform will only be calculated when alpha is bigger than 0 
+	 */
+	Curve,
+
+	/** IK Target transform will be calculated during the entire interaction. Alpha value will be always 1 */
+	None
+};
+
+USTRUCT()
+struct FContextualAnimIKTargetParams
+{
+	GENERATED_BODY()
+
+	/** IK Target definitions for each role in the interaction */
+	UPROPERTY(EditAnywhere, Category = "Defaults", meta = (TitleProperty = "Role", DisplayName = "IK Target Defs For Each Role"))
+	TArray<FContextualAnimIKTargetDefContainer> IKTargetDefsForEachRole;
+
+	/** Method used to determine the alpha of IK targets at runtime */
+	UPROPERTY(EditAnywhere, Category = "Defaults")
+	EContextualAnimIKTargetAlphaProvider AlphaProvider = EContextualAnimIKTargetAlphaProvider::AnimNotifyState;
+};
+
 // FContextualAnimRoleDefinition
 ///////////////////////////////////////////////////////////////////////
 
diff --git a/Engine/Plugins/Experimental/Animation/ContextualAnimation/Source/ContextualAnimation/Public/ContextualAnimUtilities.h b/Engine/Plugins/Experimental/Animation/ContextualAnimation/Source/ContextualAnimation/Public/ContextualAnimUtilities.h
--- a/Engine/Plugins/Experimental/Animation/ContextualAnimation/Source/ContextualAnimation/Public/ContextualAnimUtilities.h
+++ b/Engine/Plugins/Experimental/Animation/ContextualAnimation/Source/ContextualAnimation/Public/ContextualAnimUtilities.h
@@ -103,6 +103,9 @@ public:
 	UFUNCTION(BlueprintPure, Category = "Contextual Anim|Scene Bindings", meta = (DisplayName = "Get Binding By Actor"))
 	static const FContextualAnimSceneBinding& BP_SceneBindings_GetBindingByActor(const FContextualAnimSceneBindings& Bindings, const AActor* Actor);
 
+	UFUNCTION(BlueprintPure, Category = "Contextual Anim|Scene Bindings", meta = (DisplayName = "Get Primary Binding"))
+	static const FContextualAnimSceneBinding& BP_SceneBindings_GetPrimaryBinding(const FContextualAnimSceneBindings& Bindings);
+
 	UFUNCTION(BlueprintPure, Category = "Contextual Anim|Scene Bindings", meta = (DisplayName = "Get Animation From Binding"))
 	static const UAnimSequenceBase* BP_SceneBinding_GetAnimationFromBinding(const FContextualAnimSceneBindings& Bindings, const FContextualAnimSceneBinding& Binding);
 
diff --git a/Engine/Plugins/Experimental/Animation/ContextualAnimation/Source/ContextualAnimationEditor/Private/ContextualAnimEdMode.cpp b/Engine/Plugins/Experimental/Animation/ContextualAnimation/Source/ContextualAnimationEditor/Private/ContextualAnimEdMode.cpp
--- a/Engine/Plugins/Experimental/Animation/ContextualAnimation/Source/ContextualAnimationEditor/Private/ContextualAnimEdMode.cpp
+++ b/Engine/Plugins/Experimental/Animation/ContextualAnimation/Source/ContextualAnimationEditor/Private/ContextualAnimEdMode.cpp
@@ -244,31 +244,43 @@ void FContextualAnimEdMode::DrawHUD(FEditorViewportClient* ViewportClient, FView
 
 void FContextualAnimEdMode::DrawIKTargetsForBinding(FPrimitiveDrawInterface& PDI, const FContextualAnimSceneBinding& Binding) const
 {
+	const USkeletalMeshComponent* BindingSkelMeshComp = Binding.GetSkeletalMeshComponent();
+	if (!BindingSkelMeshComp)
+	{
+		return;
+	}
+
 	const TArray<FContextualAnimIKTargetDefinition>& IKTargetDefs = ViewModel->SceneBindings.GetIKTargetDefContainerFromBinding(Binding).IKTargetDefs;
+
 	for (const FContextualAnimIKTargetDefinition& IKTargetDef : IKTargetDefs)
 	{
 		if (const FContextualAnimSceneBinding* TargetBinding = ViewModel->SceneBindings.FindBindingByRole(IKTargetDef.TargetRoleName))
 		{
-			if (USkeletalMeshComponent* TargetSkelMeshComp = TargetBinding->GetSkeletalMeshComponent())
+			if (UMeshComponent* TargetMeshComp = UContextualAnimUtilities::TryGetMeshComponentWithSocket(TargetBinding->GetActor(), IKTargetDef.TargetBoneName))
 			{
+				const float Radius = 1.f;
 				const float Alpha = UAnimNotifyState_IKWindow::GetIKAlphaValue(IKTargetDef.GoalName, Binding.GetAnimMontageInstance());
+				FLinearColor Color = Alpha > 0.f ? FLinearColor(FColor::MakeRedToGreenColorFromScalar(Alpha)) : FLinearColor::White;
 
 				if (IKTargetDef.Provider == EContextualAnimIKTargetProvider::Bone)
 				{
-					//@TODO:...
+					const FTransform SourceTransform = BindingSkelMeshComp->GetBoneTransform(IKTargetDef.BoneName);
+					const FTransform TargetTransform = TargetMeshComp->GetSocketTransform(IKTargetDef.TargetBoneName);
+					
+					PDI.DrawLine(SourceTransform.GetLocation(), TargetTransform.GetLocation(), Color, SDPG_Foreground, 0.f);
+					SkeletalDebugRendering::DrawAxes(&PDI, FTransform(TargetTransform.GetLocation()), SDPG_Foreground, 0.f, Radius);
 				}
 				else if (IKTargetDef.Provider == EContextualAnimIKTargetProvider::Autogenerated)
 				{
-					const FTransform ParentTransform = TargetSkelMeshComp->GetSocketTransform(IKTargetDef.TargetBoneName);
+					const FTransform SourceTransform = BindingSkelMeshComp->GetBoneTransform(IKTargetDef.BoneName);
 
+					const FTransform ParentTransform = TargetMeshComp->GetSocketTransform(IKTargetDef.TargetBoneName);
 					const FTransform TargetTransform = ViewModel->SceneBindings.GetIKTargetTransformFromBinding(Binding, IKTargetDef.GoalName, Binding.GetAnimMontageTime()) * ParentTransform;
 
-					FLinearColor Color = Alpha > 0.f ? FLinearColor(FColor::MakeRedToGreenColorFromScalar(Alpha)) : FLinearColor::White;
-
 					FVector Start = ParentTransform.GetLocation();
 					FVector End = TargetTransform.GetLocation();
 
-					const float Radius = 1.f;
+					PDI.DrawLine(SourceTransform.GetLocation(), TargetTransform.GetLocation(), Color, SDPG_Foreground, 0.f);
 					SkeletalDebugRendering::DrawWireBone(&PDI, Start, End, Color, SDPG_Foreground, Radius);
 					SkeletalDebugRendering::DrawAxes(&PDI, FTransform(End), SDPG_Foreground, 0.f, Radius);
 				}
diff --git a/Engine/Plugins/Experimental/Animation/ContextualAnimation/Source/ContextualAnimationEditor/Private/ContextualAnimMovieSceneNotifyTrackEditor.cpp b/Engine/Plugins/Experimental/Animation/ContextualAnimation/Source/ContextualAnimationEditor/Private/ContextualAnimMovieSceneNotifyTrackEditor.cpp
--- a/Engine/Plugins/Experimental/Animation/ContextualAnimation/Source/ContextualAnimationEditor/Private/ContextualAnimMovieSceneNotifyTrackEditor.cpp
+++ b/Engine/Plugins/Experimental/Animation/ContextualAnimation/Source/ContextualAnimationEditor/Private/ContextualAnimMovieSceneNotifyTrackEditor.cpp
@@ -215,8 +215,8 @@ void FContextualAnimMovieSceneNotifyTrackEditor::BuildNewIKTargetSubMenu(FMenuBu
 		const FContextualAnimTrack* AnimTrack = SceneAsset->FindAnimTrackByAnimation(&Track->GetAnimation());
 		check(AnimTrack);
 
-		const FContextualAnimIKTargetDefContainer& IKTargets = SceneAsset->GetIKTargetDefsForRoleInSection(AnimTrack->SectionIdx, AnimTrack->Role);
-		for (const FContextualAnimIKTargetDefinition& IKTargetDef : IKTargets.IKTargetDefs)
+		const TArray<FContextualAnimIKTargetDefinition>& IKTargetDefs = SceneAsset->GetIKTargetDefsForRole(AnimTrack->Role).IKTargetDefs;
+		for (const FContextualAnimIKTargetDefinition& IKTargetDef : IKTargetDefs)
 		{
 			const FName GoalName = IKTargetDef.GoalName;
 			MenuBuilder.AddMenuEntry(
diff --git a/Engine/Plugins/Experimental/Animation/ContextualAnimation/Source/ContextualAnimationEditor/Private/ContextualAnimViewModel.cpp b/Engine/Plugins/Experimental/Animation/ContextualAnimation/Source/ContextualAnimationEditor/Private/ContextualAnimViewModel.cpp
--- a/Engine/Plugins/Experimental/Animation/ContextualAnimation/Source/ContextualAnimationEditor/Private/ContextualAnimViewModel.cpp
+++ b/Engine/Plugins/Experimental/Animation/ContextualAnimation/Source/ContextualAnimationEditor/Private/ContextualAnimViewModel.cpp
@@ -713,15 +713,17 @@ void FContextualAnimViewModel::AddNewIKTarget(const UContextualAnimNewIKTargetPa
 	IKTargetDef.TargetRoleName = Params.TargetRole;
 	IKTargetDef.TargetBoneName = Params.TargetBone.BoneName;
 
-	if (FContextualAnimIKTargetDefContainer* ContainerPtr = SceneAsset->Sections[Params.SectionIdx].RoleToIKTargetDefsMap.Find(Params.SourceRole))
+	FName Role = Params.SourceRole;
+	if (FContextualAnimIKTargetDefContainer* ContainerPtr = SceneAsset->IKTargetParams.IKTargetDefsForEachRole.FindByPredicate([Role](const FContextualAnimIKTargetDefContainer& Item) { return Item.Role == Role; }))
 	{
 		ContainerPtr->IKTargetDefs.AddUnique(IKTargetDef);
 	}
 	else
 	{
 		FContextualAnimIKTargetDefContainer Container;
-		Container.IKTargetDefs.AddUnique(IKTargetDef);
-		SceneAsset->Sections[Params.SectionIdx].RoleToIKTargetDefsMap.Add(Params.SourceRole, Container);
+		Container.Role = Role;
+		Container.IKTargetDefs.Add(IKTargetDef);
+		SceneAsset->IKTargetParams.IKTargetDefsForEachRole.Add(Container);
 	}
 
 	SceneAsset->PrecomputeData();
diff --git a/Engine/Plugins/Experimental/Animation/ContextualAnimation/Source/ContextualAnimationEditor/Private/DetailCustomizations/ContextualAnimSceneAssetDetailCustom.cpp b/Engine/Plugins/Experimental/Animation/ContextualAnimation/Source/ContextualAnimationEditor/Private/DetailCustomizations/ContextualAnimSceneAssetDetailCustom.cpp
--- a/Engine/Plugins/Experimental/Animation/ContextualAnimation/Source/ContextualAnimationEditor/Private/DetailCustomizations/ContextualAnimSceneAssetDetailCustom.cpp
+++ b/Engine/Plugins/Experimental/Animation/ContextualAnimation/Source/ContextualAnimationEditor/Private/DetailCustomizations/ContextualAnimSceneAssetDetailCustom.cpp
@@ -73,7 +73,6 @@ void FContextualAnimSceneAssetDetailCustom::CustomizeDetails(IDetailLayoutBuilde
 
 		// Show editable properties in the custom category
 		CurrentSectionCategory.AddProperty(SectionPropertyHandle->GetChildHandle(FName(TEXT("Name"))).ToSharedRef());
-		CurrentSectionCategory.AddProperty(SectionPropertyHandle->GetChildHandle(FName(TEXT("RoleToIKTargetDefsMap"))).ToSharedRef());
 		CurrentSectionCategory.AddProperty(SectionPropertyHandle->GetChildHandle(FName(TEXT("WarpPointDefinitions"))).ToSharedRef());
 		CurrentSectionCategory.AddProperty(SectionPropertyHandle->GetChildHandle(FName(TEXT("bSyncAnimations"))).ToSharedRef());
 
