Sequencer: MovieSceneSequence::IsTrackSupported now checks class permissions first and if allowed, then checks derived classes.

IsTrackSupported is no longer virtual, so derived classes will need to be changed to override IsTrackSupportedImpl

diff --git a/Engine/Plugins/Experimental/Animation/ContextualAnimation/Source/ContextualAnimationEditor/Private/ContextualAnimMovieSceneSequence.cpp b/Engine/Plugins/Experimental/Animation/ContextualAnimation/Source/ContextualAnimationEditor/Private/ContextualAnimMovieSceneSequence.cpp
--- a/Engine/Plugins/Experimental/Animation/ContextualAnimation/Source/ContextualAnimationEditor/Private/ContextualAnimMovieSceneSequence.cpp
+++ b/Engine/Plugins/Experimental/Animation/ContextualAnimation/Source/ContextualAnimationEditor/Private/ContextualAnimMovieSceneSequence.cpp
@@ -13,7 +13,7 @@ UContextualAnimMovieSceneSequence::UContextualAnimMovieSceneSequence(const FObje
 
 }
 
-ETrackSupport UContextualAnimMovieSceneSequence::IsTrackSupported(TSubclassOf<class UMovieSceneTrack> InTrackClass) const
+ETrackSupport UContextualAnimMovieSceneSequence::IsTrackSupportedImpl(TSubclassOf<class UMovieSceneTrack> InTrackClass) const
 { 
 	if (InTrackClass == UContextualAnimMovieSceneNotifyTrack::StaticClass() || 
 		InTrackClass == UContextualAnimMovieSceneTrack::StaticClass())
@@ -21,7 +21,7 @@ ETrackSupport UContextualAnimMovieSceneSequence::IsTrackSupported(TSubclassOf<cl
 		return ETrackSupport::Supported;
 	}
 	
-	return Super::IsTrackSupported(InTrackClass);
+	return Super::IsTrackSupportedImpl(InTrackClass);
 }
 
 void UContextualAnimMovieSceneSequence::Initialize(const TSharedRef<FContextualAnimViewModel>& ViewModelRef)
diff --git a/Engine/Plugins/Experimental/Animation/ContextualAnimation/Source/ContextualAnimationEditor/Public/ContextualAnimMovieSceneSequence.h b/Engine/Plugins/Experimental/Animation/ContextualAnimation/Source/ContextualAnimationEditor/Public/ContextualAnimMovieSceneSequence.h
--- a/Engine/Plugins/Experimental/Animation/ContextualAnimation/Source/ContextualAnimationEditor/Public/ContextualAnimMovieSceneSequence.h
+++ b/Engine/Plugins/Experimental/Animation/ContextualAnimation/Source/ContextualAnimationEditor/Public/ContextualAnimMovieSceneSequence.h
@@ -28,7 +28,7 @@ public:
 	virtual void UnbindPossessableObjects(const FGuid& ObjectId) override;
 	virtual void UnbindObjects(const FGuid& ObjectId, const TArray<UObject*>& InObjects, UObject* Context) override;
 	virtual void UnbindInvalidObjects(const FGuid& ObjectId, UObject* Context) override;
-	virtual ETrackSupport IsTrackSupported(TSubclassOf<class UMovieSceneTrack> InTrackClass) const;
+	virtual ETrackSupport IsTrackSupportedImpl(TSubclassOf<class UMovieSceneTrack> InTrackClass) const override;
 
 	FContextualAnimViewModel& GetViewModel() const { return *ViewModelPtr.Pin().Get(); }
 
diff --git a/Engine/Plugins/Experimental/DaySequence/Source/DaySequence/Private/DaySequence.cpp b/Engine/Plugins/Experimental/DaySequence/Source/DaySequence/Private/DaySequence.cpp
--- a/Engine/Plugins/Experimental/DaySequence/Source/DaySequence/Private/DaySequence.cpp
+++ b/Engine/Plugins/Experimental/DaySequence/Source/DaySequence/Private/DaySequence.cpp
@@ -118,13 +118,8 @@ bool UDaySequence::CanAnimateObject(UObject& InObject) const
 
 #if WITH_EDITOR
 
-ETrackSupport UDaySequence::IsTrackSupported(TSubclassOf<class UMovieSceneTrack> InTrackClass) const
+ETrackSupport UDaySequence::IsTrackSupportedImpl(TSubclassOf<class UMovieSceneTrack> InTrackClass) const
 {
-	if (!UMovieScene::IsTrackClassAllowed(InTrackClass))
-	{
-		return ETrackSupport::NotSupported;
-	}
-	
 	if (InTrackClass == UMovieScene3DAttachTrack::StaticClass() ||
 		InTrackClass == UMovieScene3DPathTrack::StaticClass() ||
 		InTrackClass == UMovieSceneAudioTrack::StaticClass() ||
@@ -140,7 +135,7 @@ ETrackSupport UDaySequence::IsTrackSupported(TSubclassOf<class UMovieSceneTrack>
 		return ETrackSupport::Supported;
 	}
 
-	return Super::IsTrackSupported(InTrackClass);
+	return Super::IsTrackSupportedImpl(InTrackClass);
 }
 
 void UDaySequence::GetAssetRegistryTags(FAssetRegistryTagsContext Context) const
diff --git a/Engine/Plugins/Experimental/DaySequence/Source/DaySequence/Public/DaySequence.h b/Engine/Plugins/Experimental/DaySequence/Source/DaySequence/Public/DaySequence.h
--- a/Engine/Plugins/Experimental/DaySequence/Source/DaySequence/Public/DaySequence.h
+++ b/Engine/Plugins/Experimental/DaySequence/Source/DaySequence/Public/DaySequence.h
@@ -59,7 +59,7 @@ public:
 	//~ End IInterface_AssetUserData Interface
 
 #if WITH_EDITOR
-	virtual ETrackSupport IsTrackSupported(TSubclassOf<class UMovieSceneTrack> InTrackClass) const override;
+	virtual ETrackSupport IsTrackSupportedImpl(TSubclassOf<class UMovieSceneTrack> InTrackClass) const override;
 	virtual void GetAssetRegistryTags(FAssetRegistryTagsContext Context) const override;
 
 	DECLARE_DELEGATE_RetVal_OneParam(void, FPostDuplicateEvent, UDaySequence*);
diff --git a/Engine/Plugins/MovieScene/ActorSequence/Source/ActorSequence/Private/ActorSequence.cpp b/Engine/Plugins/MovieScene/ActorSequence/Source/ActorSequence/Private/ActorSequence.cpp
--- a/Engine/Plugins/MovieScene/ActorSequence/Source/ActorSequence/Private/ActorSequence.cpp
+++ b/Engine/Plugins/MovieScene/ActorSequence/Source/ActorSequence/Private/ActorSequence.cpp
@@ -216,7 +216,7 @@ FText UActorSequence::GetDisplayName() const
 }
 
 
-ETrackSupport UActorSequence::IsTrackSupported(TSubclassOf<class UMovieSceneTrack> InTrackClass) const
+ETrackSupport UActorSequence::IsTrackSupportedImpl(TSubclassOf<class UMovieSceneTrack> InTrackClass) const
 {
 	if (InTrackClass == UMovieSceneAudioTrack::StaticClass() ||
 		InTrackClass == UMovieSceneEventTrack::StaticClass() ||
@@ -226,7 +226,7 @@ ETrackSupport UActorSequence::IsTrackSupported(TSubclassOf<class UMovieSceneTrac
 		return ETrackSupport::Supported;
 	}
 
-	return Super::IsTrackSupported(InTrackClass);
+	return Super::IsTrackSupportedImpl(InTrackClass);
 }
 
 #endif
diff --git a/Engine/Plugins/MovieScene/ActorSequence/Source/ActorSequence/Public/ActorSequence.h b/Engine/Plugins/MovieScene/ActorSequence/Source/ActorSequence/Public/ActorSequence.h
--- a/Engine/Plugins/MovieScene/ActorSequence/Source/ActorSequence/Public/ActorSequence.h
+++ b/Engine/Plugins/MovieScene/ActorSequence/Source/ActorSequence/Public/ActorSequence.h
@@ -32,7 +32,7 @@ public:
 
 #if WITH_EDITOR
 	virtual FText GetDisplayName() const override;
-	virtual ETrackSupport IsTrackSupported(TSubclassOf<class UMovieSceneTrack> InTrackClass) const override;
+	virtual ETrackSupport IsTrackSupportedImpl(TSubclassOf<class UMovieSceneTrack> InTrackClass) const override;
 #endif
 
 #if WITH_EDITORONLY_DATA
diff --git a/Engine/Plugins/MovieScene/TemplateSequence/Source/TemplateSequence/Private/TemplateSequence.cpp b/Engine/Plugins/MovieScene/TemplateSequence/Source/TemplateSequence/Private/TemplateSequence.cpp
--- a/Engine/Plugins/MovieScene/TemplateSequence/Source/TemplateSequence/Private/TemplateSequence.cpp
+++ b/Engine/Plugins/MovieScene/TemplateSequence/Source/TemplateSequence/Private/TemplateSequence.cpp
@@ -289,7 +289,7 @@ FText UTemplateSequence::GetDisplayName() const
 	return UMovieSceneSequence::GetDisplayName();
 }
 
-ETrackSupport UTemplateSequence::IsTrackSupported(TSubclassOf<class UMovieSceneTrack> InTrackClass) const
+ETrackSupport UTemplateSequence::IsTrackSupportedImpl(TSubclassOf<class UMovieSceneTrack> InTrackClass) const
 {
 	if (InTrackClass == UMovieSceneSkeletalAnimationTrack::StaticClass() ||
 		InTrackClass == UMovieSceneSpawnTrack::StaticClass())
@@ -297,7 +297,7 @@ ETrackSupport UTemplateSequence::IsTrackSupported(TSubclassOf<class UMovieSceneT
 		return ETrackSupport::Supported;
 	}
 
-	return Super::IsTrackSupported(InTrackClass);
+	return Super::IsTrackSupportedImpl(InTrackClass);
 }
 
 void UTemplateSequence::GetAssetRegistryTagMetadata(TMap<FName, FAssetRegistryTagMetadata>& OutMetadata) const
diff --git a/Engine/Plugins/MovieScene/TemplateSequence/Source/TemplateSequence/Public/TemplateSequence.h b/Engine/Plugins/MovieScene/TemplateSequence/Source/TemplateSequence/Public/TemplateSequence.h
--- a/Engine/Plugins/MovieScene/TemplateSequence/Source/TemplateSequence/Public/TemplateSequence.h
+++ b/Engine/Plugins/MovieScene/TemplateSequence/Source/TemplateSequence/Public/TemplateSequence.h
@@ -49,7 +49,7 @@ public:
 #if WITH_EDITOR
 	virtual FText GetDisplayName() const override;
 
-	virtual ETrackSupport IsTrackSupported(TSubclassOf<class UMovieSceneTrack> InTrackClass) const override;
+	virtual ETrackSupport IsTrackSupportedImpl(TSubclassOf<class UMovieSceneTrack> InTrackClass) const override;
 	virtual void GetAssetRegistryTagMetadata(TMap<FName, FAssetRegistryTagMetadata>& OutMetadata) const override;
 protected:
 	virtual void ThreadedPostLoadAssetRegistryTagsOverride(FPostLoadAssetRegistryTagsContext& Context) const;
diff --git a/Engine/Source/Runtime/LevelSequence/Private/LevelSequence.cpp b/Engine/Source/Runtime/LevelSequence/Private/LevelSequence.cpp
--- a/Engine/Source/Runtime/LevelSequence/Private/LevelSequence.cpp
+++ b/Engine/Source/Runtime/LevelSequence/Private/LevelSequence.cpp
@@ -137,13 +137,8 @@ bool ULevelSequence::CanAnimateObject(UObject& InObject) const
 
 #if WITH_EDITOR
 
-ETrackSupport ULevelSequence::IsTrackSupported(TSubclassOf<class UMovieSceneTrack> InTrackClass) const
+ETrackSupport ULevelSequence::IsTrackSupportedImpl(TSubclassOf<class UMovieSceneTrack> InTrackClass) const
 {
-	if (!UMovieScene::IsTrackClassAllowed(InTrackClass))
-	{
-		return ETrackSupport::NotSupported;
-	}
-
 	if (InTrackClass == UMovieScene3DAttachTrack::StaticClass() ||
 		InTrackClass == UMovieScene3DPathTrack::StaticClass() ||
 		InTrackClass == UMovieSceneAudioTrack::StaticClass() ||
@@ -164,7 +159,7 @@ ETrackSupport ULevelSequence::IsTrackSupported(TSubclassOf<class UMovieSceneTrac
 		return ETrackSupport::Supported;
 	}
 
-	return Super::IsTrackSupported(InTrackClass);
+	return Super::IsTrackSupportedImpl(InTrackClass);
 }
 
 void ULevelSequence::GetAssetRegistryTags(TArray<FAssetRegistryTag>& OutTags) const
diff --git a/Engine/Source/Runtime/LevelSequence/Public/LevelSequence.h b/Engine/Source/Runtime/LevelSequence/Public/LevelSequence.h
--- a/Engine/Source/Runtime/LevelSequence/Public/LevelSequence.h
+++ b/Engine/Source/Runtime/LevelSequence/Public/LevelSequence.h
@@ -78,7 +78,7 @@ public:
 
 
 #if WITH_EDITOR
-	LEVELSEQUENCE_API virtual ETrackSupport IsTrackSupported(TSubclassOf<class UMovieSceneTrack> InTrackClass) const override;
+	LEVELSEQUENCE_API virtual ETrackSupport IsTrackSupportedImpl(TSubclassOf<class UMovieSceneTrack> InTrackClass) const override;
 	LEVELSEQUENCE_API virtual void GetAssetRegistryTagMetadata(TMap<FName, FAssetRegistryTagMetadata>& OutMetadata) const override;
 	LEVELSEQUENCE_API virtual void GetAssetRegistryTags(FAssetRegistryTagsContext Context) const override;
 	UE_DEPRECATED(5.4, "Implement the version that takes FAssetRegistryTagsContext instead.")
diff --git a/Engine/Source/Runtime/MovieScene/Private/MovieSceneSequence.cpp b/Engine/Source/Runtime/MovieScene/Private/MovieSceneSequence.cpp
--- a/Engine/Source/Runtime/MovieScene/Private/MovieSceneSequence.cpp
+++ b/Engine/Source/Runtime/MovieScene/Private/MovieSceneSequence.cpp
@@ -396,3 +396,16 @@ UObject* UMovieSceneSequence::CreateDirectorInstance(IMovieScenePlayer& Player,
 	return CreateDirectorInstance(Player.GetSharedPlaybackState(), SequenceID);
 }
 
+#if WITH_EDITOR
+
+ETrackSupport UMovieSceneSequence::IsTrackSupported(TSubclassOf<class UMovieSceneTrack> InTrackClass) const
+{
+	if (!UMovieScene::IsTrackClassAllowed(InTrackClass))
+	{
+		return ETrackSupport::NotSupported;
+	}
+
+	return IsTrackSupportedImpl(InTrackClass); 
+}
+
+#endif
\ No newline at end of file
diff --git a/Engine/Source/Runtime/MovieScene/Public/MovieSceneSequence.h b/Engine/Source/Runtime/MovieScene/Public/MovieSceneSequence.h
--- a/Engine/Source/Runtime/MovieScene/Public/MovieSceneSequence.h
+++ b/Engine/Source/Runtime/MovieScene/Public/MovieSceneSequence.h
@@ -436,6 +436,10 @@ public:
 	/*
 	 * Sequences can determine whether they support a particular track type
 	 */
-	virtual ETrackSupport IsTrackSupported(TSubclassOf<UMovieSceneTrack> InTrackClass) const { return ETrackSupport::Default; }
+	MOVIESCENE_API ETrackSupport IsTrackSupported(TSubclassOf<UMovieSceneTrack> InTrackClass) const;
+
+	/** Subclasses can implement this to determine whether they support a particular track type */
+	virtual ETrackSupport IsTrackSupportedImpl(TSubclassOf<UMovieSceneTrack> InTrackClass) const { return ETrackSupport::Default; }
+
 #endif
 };
diff --git a/Engine/Source/Runtime/UMG/Private/Animation/WidgetAnimation.cpp b/Engine/Source/Runtime/UMG/Private/Animation/WidgetAnimation.cpp
--- a/Engine/Source/Runtime/UMG/Private/Animation/WidgetAnimation.cpp
+++ b/Engine/Source/Runtime/UMG/Private/Animation/WidgetAnimation.cpp
@@ -112,7 +112,7 @@ FText UWidgetAnimation::GetDisplayName() const
 	return bHasDisplayLabel ? FText::FromString(DisplayLabel) : Super::GetDisplayName();
 }
 
-ETrackSupport UWidgetAnimation::IsTrackSupported(TSubclassOf<class UMovieSceneTrack> InTrackClass) const
+ETrackSupport UWidgetAnimation::IsTrackSupportedImpl(TSubclassOf<class UMovieSceneTrack> InTrackClass) const
 {
 	if (InTrackClass == UMovieSceneAudioTrack::StaticClass() ||
 		InTrackClass == UMovieSceneEventTrack::StaticClass() ||
@@ -121,7 +121,7 @@ ETrackSupport UWidgetAnimation::IsTrackSupported(TSubclassOf<class UMovieSceneTr
 		return ETrackSupport::Supported;
 	}
 
-	return Super::IsTrackSupported(InTrackClass);
+	return Super::IsTrackSupportedImpl(InTrackClass);
 }
 
 const FSlateBrush* UWidgetAnimation::GetCustomBrushForBinding(FGuid BindingID) const
diff --git a/Engine/Source/Runtime/UMG/Public/Animation/WidgetAnimation.h b/Engine/Source/Runtime/UMG/Public/Animation/WidgetAnimation.h
--- a/Engine/Source/Runtime/UMG/Public/Animation/WidgetAnimation.h
+++ b/Engine/Source/Runtime/UMG/Public/Animation/WidgetAnimation.h
@@ -107,7 +107,7 @@ public:
 	
 #if WITH_EDITOR
 	virtual void IterateDynamicBindings(TFunction<void(const FGuid&, FMovieSceneDynamicBinding&)> InCallback) override;
-	virtual ETrackSupport IsTrackSupported(TSubclassOf<class UMovieSceneTrack> InTrackClass) const override;
+	virtual ETrackSupport IsTrackSupportedImpl(TSubclassOf<class UMovieSceneTrack> InTrackClass) const override;
 	virtual const FSlateBrush* GetCustomBrushForBinding(FGuid BindingID) const override;
 
 	DECLARE_MULTICAST_DELEGATE_OneParam(FFixupWidgetDynamicBindingsEvent, UWidgetAnimation*);
