Contextual Anim: Use ground movement mode when resetting movement mode at the end of the interaction if we were walking or navwalking and the current ground movement is one of those. This prevent the system from overriding the ground movement mode on AIs  if it changes while the anim is playing.

diff --git a/Engine/Plugins/Experimental/Animation/ContextualAnimation/Source/ContextualAnimation/Private/ContextualAnimSceneActorComponent.cpp b/Engine/Plugins/Experimental/Animation/ContextualAnimation/Source/ContextualAnimation/Private/ContextualAnimSceneActorComponent.cpp
--- a/Engine/Plugins/Experimental/Animation/ContextualAnimation/Source/ContextualAnimation/Private/ContextualAnimSceneActorComponent.cpp
+++ b/Engine/Plugins/Experimental/Animation/ContextualAnimation/Source/ContextualAnimation/Private/ContextualAnimSceneActorComponent.cpp
@@ -1139,7 +1139,18 @@ void UContextualAnimSceneActorComponent::RestoreMovementState(const FContextualA
 		MovementComp->bAllowPhysicsRotationDuringAnimRootMotion = CharacterPropertiesBackup.bAllowPhysicsRotationDuringAnimRootMotion;
 		MovementComp->bUseControllerDesiredRotation = CharacterPropertiesBackup.bUseControllerDesiredRotation;
 		MovementComp->bOrientRotationToMovement = CharacterPropertiesBackup.bOrientRotationToMovement;
-		MovementComp->SetMovementMode(CharacterPropertiesBackup.MovementMode);
+		if (
+			(CharacterPropertiesBackup.MovementMode == MOVE_Walking || CharacterPropertiesBackup.MovementMode == MOVE_NavWalking) &&
+			(MovementComp->GetGroundMovementMode() == MOVE_Walking || MovementComp->GetGroundMovementMode() == MOVE_NavWalking))
+		{
+			// This is mostly for AIs, we want to make sure to not override their ground movement mode if it changes while the anim is playing
+			// @TODO: We should have a more robust mechanism around this to prevent the same problem with any of the movement properties that change during the interaction.
+			MovementComp->SetMovementMode(MovementComp->GetGroundMovementMode());
+		}
+		else
+		{
+			MovementComp->SetMovementMode(CharacterPropertiesBackup.MovementMode);
+		}
 		if(UPrimitiveComponent* const UpdatedPrimitive = MovementComp->UpdatedPrimitive)
 		{
 			UpdatedPrimitive->SetSimulatePhysics(CharacterPropertiesBackup.bSimulatePhysics);
