Contextual Anim: Added optional name to identify anim sets and blueprint callable functions to get current section and set names

diff --git a/Engine/Plugins/Experimental/Animation/ContextualAnimation/Source/ContextualAnimation/Private/ContextualAnimUtilities.cpp b/Engine/Plugins/Experimental/Animation/ContextualAnimation/Source/ContextualAnimation/Private/ContextualAnimUtilities.cpp
--- a/Engine/Plugins/Experimental/Animation/ContextualAnimation/Source/ContextualAnimation/Private/ContextualAnimUtilities.cpp
+++ b/Engine/Plugins/Experimental/Animation/ContextualAnimation/Source/ContextualAnimation/Private/ContextualAnimUtilities.cpp
@@ -16,6 +16,7 @@
 #include "MotionWarpingComponent.h"
 #include "AnimNotifyState_MotionWarping.h"
 #include "RootMotionModifier.h"
+#include "ContextualAnimSceneActorComponent.h"
 
 #include UE_INLINE_GENERATED_CPP_BY_NAME(ContextualAnimUtilities)
 
@@ -521,3 +522,53 @@ FTransform UContextualAnimUtilities::BP_SceneBindings_GetAlignmentTransformFromB
 
 	return Bindings.GetSceneAsset()->GetAlignmentTransform(AnimTrack, WarpPoint.Name, EndTime) * WarpPoint.Transform;
 }
+
+void UContextualAnimUtilities::BP_SceneBindings_GetSectionAndAnimSetNames(const FContextualAnimSceneBindings& Bindings, FName& SectionName, FName& AnimSetName)
+{
+	SectionName = NAME_None;
+	AnimSetName = NAME_None;
+
+	if (!Bindings.IsValid())
+	{
+		UE_LOG(LogContextualAnim, Warning, TEXT("UContextualAnimUtilities::BP_SceneBindings_GetSectionAndAnimSetNames Failed. Reason: Invalid Bindings. SceneAsset: %s"), *GetNameSafe(Bindings.GetSceneAsset()));
+		return;
+	}
+
+	const FContextualAnimSceneSection* Section = Bindings.GetSceneAsset()->GetSection(Bindings.GetSectionIdx());
+	if (Section == nullptr)
+	{
+		UE_LOG(LogContextualAnim, Warning, TEXT("UContextualAnimUtilities::BP_SceneBindings_GetSectionAndAnimSetNames Failed. Reason: '%d' is not a valid section idx in '%s'"), 
+			Bindings.GetSectionIdx(), *GetNameSafe(Bindings.GetSceneAsset()));
+		return;
+	}
+
+	const FContextualAnimSet* AnimSet = Section->GetAnimSet(Bindings.GetAnimSetIdx());
+	if (AnimSet == nullptr)
+	{
+		UE_LOG(LogContextualAnim, Warning, TEXT("UContextualAnimUtilities::BP_SceneBindings_GetSectionAndAnimSetNames Failed. Reason: '%d' is not a valid anim set idx in '%s'"),
+			Bindings.GetAnimSetIdx(), *GetNameSafe(Bindings.GetSceneAsset()));
+		return;
+	}
+
+	SectionName = Section->GetName();
+	AnimSetName = AnimSet->Name;
+}
+
+void UContextualAnimUtilities::BP_SceneBindingContext_GetCurrentSectionAndAnimSetNames(const FContextualAnimSceneBindingContext& BindingContext, FName& SectionName, FName& AnimSetName)
+{
+	SectionName = NAME_None;
+	AnimSetName = NAME_None;
+
+	const UContextualAnimSceneActorComponent* SceneComp = BindingContext.GetSceneActorComponent();
+	if (SceneComp == nullptr)
+	{
+		UE_LOG(LogContextualAnim, Warning, TEXT("UContextualAnimUtilities::BP_SceneBindingContext_GetCurrentSectionAndAnimSetNames Failed. Reason: Missing SceneActorComp. Actor: %s"), *GetNameSafe(BindingContext.GetActor()));
+		return;
+	}
+
+	const FContextualAnimSceneBindings& Bindings = SceneComp->GetBindings();
+	if (Bindings.IsValid())
+	{
+		UContextualAnimUtilities::BP_SceneBindings_GetSectionAndAnimSetNames(Bindings, SectionName, AnimSetName);
+	}
+}
\ No newline at end of file
diff --git a/Engine/Plugins/Experimental/Animation/ContextualAnimation/Source/ContextualAnimation/Public/ContextualAnimSceneAsset.h b/Engine/Plugins/Experimental/Animation/ContextualAnimation/Source/ContextualAnimation/Public/ContextualAnimSceneAsset.h
--- a/Engine/Plugins/Experimental/Animation/ContextualAnimation/Source/ContextualAnimation/Public/ContextualAnimSceneAsset.h
+++ b/Engine/Plugins/Experimental/Animation/ContextualAnimation/Source/ContextualAnimation/Public/ContextualAnimSceneAsset.h
@@ -82,6 +82,10 @@ struct CONTEXTUALANIMATION_API FContextualAnimSet
 	UPROPERTY(EditAnywhere, Category = "Defaults")
 	TMap<FName, FTransform> WarpPoints;
 
+	/** Optional name to identify this set */
+	UPROPERTY(EditAnywhere, Category = "Defaults")
+	FName Name = NAME_None;
+
 	/** Used by the selection mechanism to 'break the tie' when multiple Sets can be selected */
 	UPROPERTY(EditAnywhere, Category = "Defaults", meta = (ClampMin = "0", UIMin = "0", ClampMax = "1", UIMax = "1"))
 	float RandomWeight = 1.f;
diff --git a/Engine/Plugins/Experimental/Animation/ContextualAnimation/Source/ContextualAnimation/Public/ContextualAnimUtilities.h b/Engine/Plugins/Experimental/Animation/ContextualAnimation/Source/ContextualAnimation/Public/ContextualAnimUtilities.h
--- a/Engine/Plugins/Experimental/Animation/ContextualAnimation/Source/ContextualAnimation/Public/ContextualAnimUtilities.h
+++ b/Engine/Plugins/Experimental/Animation/ContextualAnimation/Source/ContextualAnimation/Public/ContextualAnimUtilities.h
@@ -118,6 +118,9 @@ public:
 	UFUNCTION(BlueprintPure, Category = "Contextual Anim|Scene Bindings", meta = (DisplayName = "Get Section And AnimSet Indices"))
 	static void BP_SceneBindings_GetSectionAndAnimSetIndices(const FContextualAnimSceneBindings& Bindings, int32& SectionIdx, int32& AnimSetIdx);
 
+	UFUNCTION(BlueprintPure, Category = "Contextual Anim|Scene Bindings", meta = (DisplayName = "Get Section And AnimSet Names"))
+	static void BP_SceneBindings_GetSectionAndAnimSetNames(const FContextualAnimSceneBindings& Bindings, FName& SectionName, FName& AnimSetName);
+
 	UFUNCTION(BlueprintPure, Category = "Contextual Anim|Scene Bindings", meta = (DisplayName = "Get Alignment Transform For Role Relative To Other Role"))
 	static FTransform BP_SceneBindings_GetAlignmentTransformForRoleRelativeToOtherRole(const FContextualAnimSceneBindings& Bindings, FName Role, FName RelativeToRole, float Time);
 
@@ -157,6 +160,9 @@ public:
 	UFUNCTION(BlueprintPure, Category = "Contextual Anim|Scene Binding Context", meta = (DisplayName = "Has Any Matching GameplayTags"))
 	static bool BP_SceneBindingContext_HasAnyMatchingGameplayTags(const FContextualAnimSceneBindingContext& BindingContext, const FGameplayTagContainer& TagContainer) { return BindingContext.HasAnyMatchingGameplayTags(TagContainer); }
 
+	UFUNCTION(BlueprintPure, Category = "Contextual Anim|Scene Binding Context", meta = (DisplayName = "Get Current Section and Anim Set"))
+	static void BP_SceneBindingContext_GetCurrentSectionAndAnimSetNames(const FContextualAnimSceneBindingContext& BindingContext, FName& SectionName, FName& AnimSetName);
+
 
 	// FContextualAnimSceneBinding Blueprint Interface
 	//------------------------------------------------------------------------------------------
diff --git a/Engine/Plugins/Experimental/Animation/ContextualAnimation/Source/ContextualAnimationEditor/Private/DetailCustomizations/ContextualAnimSceneAssetDetailCustom.cpp b/Engine/Plugins/Experimental/Animation/ContextualAnimation/Source/ContextualAnimationEditor/Private/DetailCustomizations/ContextualAnimSceneAssetDetailCustom.cpp
--- a/Engine/Plugins/Experimental/Animation/ContextualAnimation/Source/ContextualAnimationEditor/Private/DetailCustomizations/ContextualAnimSceneAssetDetailCustom.cpp
+++ b/Engine/Plugins/Experimental/Animation/ContextualAnimation/Source/ContextualAnimationEditor/Private/DetailCustomizations/ContextualAnimSceneAssetDetailCustom.cpp
@@ -127,6 +127,7 @@ void FContextualAnimSceneAssetDetailCustom::CustomizeDetails(IDetailLayoutBuilde
 			TSharedRef<IPropertyHandle> AnimSetPropertyHandle = AnimSetArrayPropertyHandle->GetChildHandle(AnimSetIdx).ToSharedRef();
 
 			// Show general AnimSet properties in a custom category
+			CurrentAnimSetCategory.AddProperty(AnimSetPropertyHandle->GetChildHandle(FName(TEXT("Name"))).ToSharedRef());
 			CurrentAnimSetCategory.AddProperty(AnimSetPropertyHandle->GetChildHandle(FName(TEXT("RandomWeight"))).ToSharedRef());
 
 			// Add button to remove the entire set
