Contexual Anim - Added option to stop animation on all the actors when stopping an interaction early

diff --git a/Engine/Plugins/Experimental/Animation/ContextualAnimation/Source/ContextualAnimation/Private/AnimNotifyState_EarlyOutContextualAnimWindow.cpp b/Engine/Plugins/Experimental/Animation/ContextualAnimation/Source/ContextualAnimation/Private/AnimNotifyState_EarlyOutContextualAnimWindow.cpp
--- a/Engine/Plugins/Experimental/Animation/ContextualAnimation/Source/ContextualAnimation/Private/AnimNotifyState_EarlyOutContextualAnimWindow.cpp
+++ b/Engine/Plugins/Experimental/Animation/ContextualAnimation/Source/ContextualAnimation/Private/AnimNotifyState_EarlyOutContextualAnimWindow.cpp
@@ -28,7 +28,7 @@ void UAnimNotifyState_EarlyOutContextualAnimWindow::NotifyTick(USkeletalMeshComp
 					UContextualAnimSceneActorComponent* Comp = CharacterOwner->FindComponentByClass<UContextualAnimSceneActorComponent>();
 					if (Comp)
 					{
-						Comp->EarlyOutContextualAnimScene();
+						Comp->EarlyOutContextualAnimScene(bStopEveryone);
 					}
 				}
 			}
diff --git a/Engine/Plugins/Experimental/Animation/ContextualAnimation/Source/ContextualAnimation/Private/ContextualAnimSceneActorComponent.cpp b/Engine/Plugins/Experimental/Animation/ContextualAnimation/Source/ContextualAnimation/Private/ContextualAnimSceneActorComponent.cpp
--- a/Engine/Plugins/Experimental/Animation/ContextualAnimation/Source/ContextualAnimation/Private/ContextualAnimSceneActorComponent.cpp
+++ b/Engine/Plugins/Experimental/Animation/ContextualAnimation/Source/ContextualAnimation/Private/ContextualAnimSceneActorComponent.cpp
@@ -439,6 +439,7 @@ bool UContextualAnimSceneActorComponent::TransitionContextualAnimScene(FName Sec
 			RepTransitionData.Id = Bindings.GetID();
 			RepTransitionData.SectionIdx = SectionIdx;
 			RepTransitionData.AnimSetIdx = AnimSetIdx;
+			RepTransitionData.bStopEveryone = false;
 			RepTransitionData.WarpPoints = MoveTemp(WarpPoints);
 			RepTransitionData.ExternalWarpTargets = ExternalWarpTargets;
 			RepTransitionData.IncrementRepCounter();
@@ -593,6 +594,7 @@ bool UContextualAnimSceneActorComponent::TransitionSingleActor(int32 SectionIdx,
 				RepTransitionSingleActorData.Id = Bindings.GetID();
 				RepTransitionSingleActorData.SectionIdx = SectionIdx;
 				RepTransitionSingleActorData.AnimSetIdx = AnimSetIdx;
+				RepTransitionSingleActorData.bStopEveryone = false;
 				RepTransitionSingleActorData.WarpPoints = MoveTemp(WarpPoints);
 				RepTransitionSingleActorData.ExternalWarpTargets = ExternalWarpTargets;
 				RepTransitionSingleActorData.IncrementRepCounter();
@@ -656,6 +658,25 @@ void UContextualAnimSceneActorComponent::OnRep_RepTransitionSingleActor()
 			// It is more likely this will do nothing since we listen to montage end also on Simulated Proxies to 'predict' the end of the interaction.
 			if (RepTransitionSingleActorData.Id == Bindings.GetID())
 			{
+				// Ensure that all other actors stop their animations if requested.
+				if (RepTransitionSingleActorData.bStopEveryone)
+				{
+					// @TODO: We copy bindings other we would be iterating over an array that would change in size due to LeveScene() updates the other's actor bindings.
+					// This should go away once we add the concept of a manager and keep bindings in a single place
+					FContextualAnimSceneBindings TempBindings = Bindings;
+
+					for (const FContextualAnimSceneBinding& OtherBinding : TempBindings)
+					{
+						if (OtherBinding.GetActor() != GetOwner())
+						{
+							if (UContextualAnimSceneActorComponent* Comp = OtherBinding.GetSceneActorComponent())
+							{
+								Comp->LeaveScene();
+							}
+						}
+					}
+				}
+
 				LeaveScene();
 			}
 		}
@@ -747,19 +768,38 @@ bool UContextualAnimSceneActorComponent::ServerStartContextualAnimScene_Validate
 	return true;
 }
 
-void UContextualAnimSceneActorComponent::EarlyOutContextualAnimScene()
+void UContextualAnimSceneActorComponent::EarlyOutContextualAnimScene(bool bStopEveryone)
 {
 	if (const FContextualAnimSceneBinding* Binding = Bindings.FindBindingByActor(GetOwner()))
 	{
 		const UAnimInstance* AnimInstance = Binding->GetAnimInstance();
 		const UAnimMontage* ActiveMontage = AnimInstance ? AnimInstance->GetCurrentActiveMontage() : nullptr;
-		UE_LOG(LogContextualAnim, Verbose, TEXT("%-21s UContextualAnimSceneActorComponent::EarlyOutContextualAnimScene Actor: %s ActiveMontage: %s"),
-			*UEnum::GetValueAsString(TEXT("Engine.ENetRole"), GetOwner()->GetLocalRole()), *GetNameSafe(GetOwner()), *GetNameSafe(ActiveMontage));
+		UE_LOG(LogContextualAnim, Verbose, TEXT("%-21s UContextualAnimSceneActorComponent::EarlyOutContextualAnimScene Actor: %s ActiveMontage: %s bStopEveryone: %d"),
+			*UEnum::GetValueAsString(TEXT("Engine.ENetRole"), GetOwner()->GetLocalRole()), *GetNameSafe(GetOwner()), *GetNameSafe(ActiveMontage), bStopEveryone);
 
 		if (Bindings.GetAnimTrackFromBinding(*Binding).Animation == ActiveMontage)
 		{
 			const uint8 BindingsId = Bindings.GetID();
 
+			// Ensure that all other actors stop their animations if requested. We do this first since the calling LeaveScene() clears all the bindings.
+			if (bStopEveryone)
+			{
+				// @TODO: We copy bindings other we would be iterating over an array that would change in size due to LeveScene() updates the other's actor bindings. 
+				// This should go away once we add the concept of a manager and keep bindings in a single place
+				FContextualAnimSceneBindings TempBindings = Bindings;
+
+				for (const FContextualAnimSceneBinding& OtherBinding : TempBindings)
+				{
+					if (OtherBinding.GetActor() != GetOwner())
+					{
+						if (UContextualAnimSceneActorComponent* Comp = OtherBinding.GetSceneActorComponent())
+						{
+							Comp->LeaveScene();
+						}
+					}
+				}
+			}
+
 			// Stop animation.
 			LeaveScene();
 
@@ -769,6 +809,8 @@ void UContextualAnimSceneActorComponent::EarlyOutContextualAnimScene()
 				RepTransitionSingleActorData.Id = BindingsId;
 				RepTransitionSingleActorData.SectionIdx = MAX_uint8;
 				RepTransitionSingleActorData.AnimSetIdx = MAX_uint8;
+				RepTransitionSingleActorData.bStopEveryone = bStopEveryone;
+				RepTransitionSingleActorData.WarpPoints.Reset();
 				RepTransitionSingleActorData.ExternalWarpTargets.Reset();
 				RepTransitionSingleActorData.IncrementRepCounter();
 
@@ -786,18 +828,18 @@ void UContextualAnimSceneActorComponent::EarlyOutContextualAnimScene()
 			// If local player, tell the server to stop the animation too
 			else if (GetOwner()->GetLocalRole() == ROLE_AutonomousProxy)
 			{
-				ServerEarlyOutContextualAnimScene();
+				ServerEarlyOutContextualAnimScene(bStopEveryone);
 			}
 		}
 	}
 }
 
-void UContextualAnimSceneActorComponent::ServerEarlyOutContextualAnimScene_Implementation()
+void UContextualAnimSceneActorComponent::ServerEarlyOutContextualAnimScene_Implementation(bool bStopEveryone)
 {
-	EarlyOutContextualAnimScene();
+	EarlyOutContextualAnimScene(bStopEveryone);
 }
 
-bool UContextualAnimSceneActorComponent::ServerEarlyOutContextualAnimScene_Validate()
+bool UContextualAnimSceneActorComponent::ServerEarlyOutContextualAnimScene_Validate(bool bStopEveryone)
 {
 	return true;
 }
@@ -1198,6 +1240,7 @@ void UContextualAnimSceneActorComponent::OnMontageBlendingOut(UAnimMontage* Mont
 			RepTransitionSingleActorData.Id = BindingsId;
 			RepTransitionSingleActorData.SectionIdx = MAX_uint8;
 			RepTransitionSingleActorData.AnimSetIdx = MAX_uint8;
+			RepTransitionSingleActorData.bStopEveryone = false;
 			RepTransitionSingleActorData.WarpPoints.Reset();
 			RepTransitionSingleActorData.ExternalWarpTargets.Reset();
 			RepTransitionSingleActorData.IncrementRepCounter();
diff --git a/Engine/Plugins/Experimental/Animation/ContextualAnimation/Source/ContextualAnimation/Public/AnimNotifyState_EarlyOutContextualAnimWindow.h b/Engine/Plugins/Experimental/Animation/ContextualAnimation/Source/ContextualAnimation/Public/AnimNotifyState_EarlyOutContextualAnimWindow.h
--- a/Engine/Plugins/Experimental/Animation/ContextualAnimation/Source/ContextualAnimation/Public/AnimNotifyState_EarlyOutContextualAnimWindow.h
+++ b/Engine/Plugins/Experimental/Animation/ContextualAnimation/Source/ContextualAnimation/Public/AnimNotifyState_EarlyOutContextualAnimWindow.h
@@ -13,6 +13,10 @@ class CONTEXTUALANIMATION_API UAnimNotifyState_EarlyOutContextualAnimWindow : pu
 
 public:
 
+	/** Whether to force all the actors in the interaction to exit the interaction or only us */
+	UPROPERTY(EditAnywhere, Category = Config)
+	bool bStopEveryone = false;
+
 	UAnimNotifyState_EarlyOutContextualAnimWindow(const FObjectInitializer& ObjectInitializer);
 
 	virtual void NotifyTick(USkeletalMeshComponent* MeshComp, UAnimSequenceBase* Animation, float FrameDeltaTime, const FAnimNotifyEventReference& EventReference) override;
diff --git a/Engine/Plugins/Experimental/Animation/ContextualAnimation/Source/ContextualAnimation/Public/ContextualAnimSceneActorComponent.h b/Engine/Plugins/Experimental/Animation/ContextualAnimation/Source/ContextualAnimation/Public/ContextualAnimSceneActorComponent.h
--- a/Engine/Plugins/Experimental/Animation/ContextualAnimation/Source/ContextualAnimation/Public/ContextualAnimSceneActorComponent.h
+++ b/Engine/Plugins/Experimental/Animation/ContextualAnimation/Source/ContextualAnimation/Public/ContextualAnimSceneActorComponent.h
@@ -126,6 +126,9 @@ struct FContextualAnimRepTransitionData : public FContextualAnimRepData
 	UPROPERTY()
 	uint8 AnimSetIdx = 0;
 
+	UPROPERTY()
+	bool bStopEveryone = false;
+
 	UPROPERTY()
 	TArray<FContextualAnimWarpPoint> WarpPoints;
 
@@ -138,6 +141,7 @@ struct FContextualAnimRepTransitionData : public FContextualAnimRepData
 		Id = 0;
 		SectionIdx = 0;
 		AnimSetIdx = 0;
+		bStopEveryone = false;
 		WarpPoints.Reset();
 		ExternalWarpTargets.Reset();
 	}
@@ -220,7 +224,7 @@ public:
 	bool TransitionSingleActor(int32 SectionIdx, int32 AnimSetIdx, const TArray<FContextualAnimWarpTarget>& ExternalWarpTargets);
 
 	UFUNCTION(BlueprintCallable, Category = "Contextual Anim|Scene Actor Component")
-	void EarlyOutContextualAnimScene();
+	void EarlyOutContextualAnimScene(bool bStopEveryone = false);
 
 	bool IsOwnerLocallyControlled() const;
 
@@ -336,7 +340,7 @@ protected:
 	void ServerStartContextualAnimScene(const FContextualAnimSceneBindings& InBindings);
 
 	UFUNCTION(Server, Reliable, WithValidation)
-	void ServerEarlyOutContextualAnimScene();
+	void ServerEarlyOutContextualAnimScene(bool bStopEveryone);
 
 	virtual void PlayAnimation_Internal(UAnimSequenceBase* Animation, float StartTime, bool bSyncPlaybackTime);
 
