Contextual Anim: Stop any playing montage when an actor starts an interaction or transition to a new section with no animation for that actor.

diff --git a/Engine/Plugins/Experimental/Animation/ContextualAnimation/Source/ContextualAnimation/Private/ContextualAnimSceneActorComponent.cpp b/Engine/Plugins/Experimental/Animation/ContextualAnimation/Source/ContextualAnimation/Private/ContextualAnimSceneActorComponent.cpp
--- a/Engine/Plugins/Experimental/Animation/ContextualAnimation/Source/ContextualAnimation/Private/ContextualAnimSceneActorComponent.cpp
+++ b/Engine/Plugins/Experimental/Animation/ContextualAnimation/Source/ContextualAnimation/Private/ContextualAnimSceneActorComponent.cpp
@@ -136,37 +136,45 @@ void UContextualAnimSceneActorComponent::PlayAnimation_Internal(UAnimSequenceBas
 		// Cache AnimInstance so we don't have to look for it in the bindings in subsequent uses.
 		OwnerAnimInstance = AnimInstance;
 
-		//@TODO: Add support for dynamic montage
-		UAnimMontage* AnimMontage = Cast<UAnimMontage>(Animation);
+		if (Animation)
+		{
+			//@TODO: Add support for dynamic montage
+			UAnimMontage* AnimMontage = Cast<UAnimMontage>(Animation);
 
-		// Keep track of this animation. Used as guarding mechanism in OnMonageBlendingOut to decide if is safe to leave the scene
-		AnimsPlayed.Add(AnimMontage);
+			// Keep track of this animation. Used as guarding mechanism in OnMonageBlendingOut to decide if is safe to leave the scene
+			AnimsPlayed.Add(AnimMontage);
 
-		AnimInstance->Montage_Play(AnimMontage, 1.f, EMontagePlayReturnType::MontageLength, StartTime);
+			AnimInstance->Montage_Play(AnimMontage, 1.f, EMontagePlayReturnType::MontageLength, StartTime);
 
-		AnimInstance->OnMontageBlendingOut.AddUniqueDynamic(this, &UContextualAnimSceneActorComponent::OnMontageBlendingOut);
-		AnimInstance->OnPlayMontageNotifyBegin.AddUniqueDynamic(this, &UContextualAnimSceneActorComponent::OnPlayMontageNotifyBegin);
+			AnimInstance->OnMontageBlendingOut.AddUniqueDynamic(this, &UContextualAnimSceneActorComponent::OnMontageBlendingOut);
+			AnimInstance->OnPlayMontageNotifyBegin.AddUniqueDynamic(this, &UContextualAnimSceneActorComponent::OnPlayMontageNotifyBegin);
 
-		if (bSyncPlaybackTime)
-		{
-			if (FAnimMontageInstance* MontageInstance = AnimInstance->GetActiveMontageInstance())
+			if (bSyncPlaybackTime)
 			{
-				if (const FContextualAnimSceneBinding* SyncLeader = Bindings.GetSyncLeader())
+				if (FAnimMontageInstance* MontageInstance = AnimInstance->GetActiveMontageInstance())
 				{
-					if (SyncLeader->GetActor() != GetOwner())
+					if (const FContextualAnimSceneBinding* SyncLeader = Bindings.GetSyncLeader())
 					{
-						FAnimMontageInstance* LeaderMontageInstance = SyncLeader->GetAnimMontageInstance();
-						if (LeaderMontageInstance && LeaderMontageInstance->Montage == Bindings.GetAnimTrackFromBinding(*SyncLeader).Animation && MontageInstance->GetMontageSyncLeader() == nullptr)
+						if (SyncLeader->GetActor() != GetOwner())
 						{
-							UE_LOG(LogContextualAnim, VeryVerbose, TEXT("%-21s \t\tUContextualAnimSceneActorComponent::PlayAnimation_Internal Syncing Animation. Actor: %s Anim: %s StartTime: %f bSyncPlaybackTime: %d"),
-								*UEnum::GetValueAsString(TEXT("Engine.ENetRole"), GetOwner()->GetLocalRole()), *GetNameSafe(GetOwner()), *GetNameSafe(Animation), StartTime, bSyncPlaybackTime);
+							FAnimMontageInstance* LeaderMontageInstance = SyncLeader->GetAnimMontageInstance();
+							if (LeaderMontageInstance && LeaderMontageInstance->Montage == Bindings.GetAnimTrackFromBinding(*SyncLeader).Animation && MontageInstance->GetMontageSyncLeader() == nullptr)
+							{
+								UE_LOG(LogContextualAnim, VeryVerbose, TEXT("%-21s \t\tUContextualAnimSceneActorComponent::PlayAnimation_Internal Syncing Animation. Actor: %s Anim: %s StartTime: %f bSyncPlaybackTime: %d"),
+									*UEnum::GetValueAsString(TEXT("Engine.ENetRole"), GetOwner()->GetLocalRole()), *GetNameSafe(GetOwner()), *GetNameSafe(Animation), StartTime, bSyncPlaybackTime);
 
-							MontageInstance->MontageSync_Follow(LeaderMontageInstance);
+								MontageInstance->MontageSync_Follow(LeaderMontageInstance);
+							}
 						}
 					}
 				}
 			}
 		}
+		// if there is no animation for this actor we have to make sure we stop any currently playing animation.
+		else
+		{
+			AnimInstance->StopAllMontages(0.f);
+		}
 	}
 
 	USkeletalMeshComponent* SkelMeshComp = UContextualAnimUtilities::TryGetSkeletalMeshComponent(GetOwner());
