Added a new delegate that fires when a montage is blending out, only if the anim that's blending out was part of the statetree.

diff --git a/Engine/Plugins/Experimental/Animation/ContextualAnimation/Source/ContextualAnimation/Private/ContextualAnimSceneActorComponent.cpp b/Engine/Plugins/Experimental/Animation/ContextualAnimation/Source/ContextualAnimation/Private/ContextualAnimSceneActorComponent.cpp
--- a/Engine/Plugins/Experimental/Animation/ContextualAnimation/Source/ContextualAnimation/Private/ContextualAnimSceneActorComponent.cpp
+++ b/Engine/Plugins/Experimental/Animation/ContextualAnimation/Source/ContextualAnimation/Private/ContextualAnimSceneActorComponent.cpp
@@ -1225,7 +1225,13 @@ void UContextualAnimSceneActorComponent::OnMontageBlendingOut(UAnimMontage* Mont
 
 	if (const FContextualAnimSceneBinding* Binding = Bindings.FindBindingByActor(GetOwner()))
 	{
-		AnimsPlayed.RemoveSingleSwap(Montage);
+		//We fire this delegate only when an anim that's part of the state tree plays so we can allow player to play montages during SO interactions (if the SO allows it).
+		int32 NumAnimsRemoved = AnimsPlayed.RemoveSingleSwap(Montage);
+		if (NumAnimsRemoved > 0)
+		{
+			OnMontageBlendingOutDelegate.Broadcast(this, Montage, bInterrupted);
+		}
+
 		if (AnimsPlayed.Num() > 0)
 		{
 			UE_LOG(LogContextualAnim, Verbose, TEXT("%-21s \tUContextualAnimSceneActorComponent::OnMontageBlendingOut AnimsPlayed Num: %d"), *UEnum::GetValueAsString(TEXT("Engine.ENetRole"), GetOwner()->GetLocalRole()), AnimsPlayed.Num());
diff --git a/Engine/Plugins/Experimental/Animation/ContextualAnimation/Source/ContextualAnimation/Public/ContextualAnimSceneActorComponent.h b/Engine/Plugins/Experimental/Animation/ContextualAnimation/Source/ContextualAnimation/Public/ContextualAnimSceneActorComponent.h
--- a/Engine/Plugins/Experimental/Animation/ContextualAnimation/Source/ContextualAnimation/Public/ContextualAnimSceneActorComponent.h
+++ b/Engine/Plugins/Experimental/Animation/ContextualAnimation/Source/ContextualAnimation/Public/ContextualAnimSceneActorComponent.h
@@ -15,6 +15,7 @@ class UAnimSequenceBase;
 
 DECLARE_DYNAMIC_MULTICAST_DELEGATE_OneParam(FContextualAnimSceneActorCompDelegate, class UContextualAnimSceneActorComponent*, SceneActorComponent);
 DECLARE_DYNAMIC_MULTICAST_DELEGATE_TwoParams(FContextualAnimPlayMontageNotifyBeginDelegate, class UContextualAnimSceneActorComponent*, SceneActorComponent, FName, NotifyName);
+DECLARE_DYNAMIC_MULTICAST_DELEGATE_ThreeParams(FContextualAnimOnMontageBlendingOutDelegate, class UContextualAnimSceneActorComponent*, SceneComponent, UAnimMontage* const, BlendingOutMontage, bool, bInterrupted);
 
 USTRUCT(BlueprintType)
 struct FContextualAnimWarpTarget
@@ -165,6 +166,9 @@ public:
 	UPROPERTY(BlueprintAssignable, Category = "Events")
 	FContextualAnimPlayMontageNotifyBeginDelegate OnPlayMontageNotifyBeginDelegate;
 
+	UPROPERTY(BlueprintAssignable, Category = "Events")
+	FContextualAnimOnMontageBlendingOutDelegate OnMontageBlendingOutDelegate;
+
 	UPROPERTY(EditAnywhere, BlueprintReadWrite, Category = "Config")
 	TObjectPtr<class UContextualAnimSceneAsset> SceneAsset;
 
