Added a bool to avoid changing movement mode and avoid chaning character rotation state in CAS anim tracks:
- this is required to be able to have CAS running on the upper body while the character is still free to use their lower body/move around

diff --git a/Engine/Plugins/Experimental/Animation/ContextualAnimation/Source/ContextualAnimation/Private/ContextualAnimSceneActorComponent.cpp b/Engine/Plugins/Experimental/Animation/ContextualAnimation/Source/ContextualAnimation/Private/ContextualAnimSceneActorComponent.cpp
--- a/Engine/Plugins/Experimental/Animation/ContextualAnimation/Source/ContextualAnimation/Private/ContextualAnimSceneActorComponent.cpp
+++ b/Engine/Plugins/Experimental/Animation/ContextualAnimation/Source/ContextualAnimation/Private/ContextualAnimSceneActorComponent.cpp
@@ -344,7 +344,7 @@ void UContextualAnimSceneActorComponent::LateJoinScene(const FContextualAnimScen
 
 		SetCollisionState(*Binding);
 
-		SetMovementState(*Binding, AnimTrack->MovementMode, AnimTrack->CustomMovementMode);
+		SetMovementState(*Binding, *AnimTrack);
 
 		OnLateJoinScene(*Binding, SectionIdx, AnimSetIdx);
 
@@ -528,11 +528,14 @@ void UContextualAnimSceneActorComponent::HandleTransitionSelf(int32 NewSectionId
 
 	AddOrUpdateWarpTargets(NewSectionIdx, NewAnimSetIdx, WarpPoints, ExternalWarpTargets);
 
-	if (UCharacterMovementComponent* MovementComp = Binding.GetCharacterMovementComponent())
+	if (AnimTrack.bChangeMovementMode)
 	{
-		if (MovementComp->MovementMode != AnimTrack.MovementMode)
+		if (UCharacterMovementComponent* MovementComp = Binding.GetCharacterMovementComponent())
 		{
-			MovementComp->SetMovementMode(AnimTrack.MovementMode);
+			if (MovementComp->MovementMode != AnimTrack.MovementMode)
+			{
+				MovementComp->SetMovementMode(AnimTrack.MovementMode);
+			}
 		}
 	}
 
@@ -589,11 +592,14 @@ bool UContextualAnimSceneActorComponent::TransitionSingleActor(int32 SectionIdx,
 
 				AddOrUpdateWarpTargets(SectionIdx, AnimSetIdx, WarpPoints, ExternalWarpTargets);
 
-				if (UCharacterMovementComponent* MovementComp = OwnerBinding->GetCharacterMovementComponent())
+				if (AnimTrack->bChangeMovementMode)
 				{
-					if (MovementComp->MovementMode != AnimTrack->MovementMode)
+					if (UCharacterMovementComponent* MovementComp = OwnerBinding->GetCharacterMovementComponent())
 					{
-						MovementComp->SetMovementMode(AnimTrack->MovementMode);
+						if (MovementComp->MovementMode != AnimTrack->MovementMode)
+						{
+							MovementComp->SetMovementMode(AnimTrack->MovementMode);
+						}
 					}
 				}
 
@@ -646,11 +652,14 @@ void UContextualAnimSceneActorComponent::OnRep_RepTransitionSingleActor()
 
 					AddOrUpdateWarpTargets(RepTransitionSingleActorData.SectionIdx, RepTransitionSingleActorData.AnimSetIdx, RepTransitionSingleActorData.WarpPoints, RepTransitionSingleActorData.ExternalWarpTargets);
 
-					if (UCharacterMovementComponent* MovementComp = OwnerBinding->GetCharacterMovementComponent())
+					if (AnimTrack->bChangeMovementMode)
 					{
-						if (MovementComp->MovementMode != AnimTrack->MovementMode)
+						if (UCharacterMovementComponent* MovementComp = OwnerBinding->GetCharacterMovementComponent())
 						{
-							MovementComp->SetMovementMode(AnimTrack->MovementMode);
+							if (MovementComp->MovementMode != AnimTrack->MovementMode)
+							{
+								MovementComp->SetMovementMode(AnimTrack->MovementMode);
+							}
 						}
 					}
 
@@ -1026,7 +1035,7 @@ void UContextualAnimSceneActorComponent::JoinScene(const FContextualAnimSceneBin
 
 		SetCollisionState(*Binding);
 
-		SetMovementState(*Binding, AnimTrack.MovementMode, AnimTrack.CustomMovementMode);
+		SetMovementState(*Binding, AnimTrack);
 
 		OnJoinScene(*Binding);
 
@@ -1134,15 +1143,11 @@ void UContextualAnimSceneActorComponent::OnLeaveScene(const FContextualAnimScene
 	// For derived classes to override.
 }
 
-void UContextualAnimSceneActorComponent::SetMovementState(const FContextualAnimSceneBinding& Binding, EMovementMode DesiredMoveMode, uint8 CustomMode)
+void UContextualAnimSceneActorComponent::SetMovementState(const FContextualAnimSceneBinding& Binding, const FContextualAnimTrack& AnimTrack)
 {
 	if (UCharacterMovementComponent* MovementComp = Binding.GetCharacterMovementComponent())
 	{
 		// Save movement state before the interaction starts so we can restore it when it ends
-		CharacterPropertiesBackup.bAllowPhysicsRotationDuringAnimRootMotion = MovementComp->bAllowPhysicsRotationDuringAnimRootMotion;
-		CharacterPropertiesBackup.bUseControllerDesiredRotation = MovementComp->bUseControllerDesiredRotation;
-		CharacterPropertiesBackup.bOrientRotationToMovement = MovementComp->bOrientRotationToMovement;
-		CharacterPropertiesBackup.MovementMode = MovementComp->MovementMode;
 		CharacterPropertiesBackup.bSimulatePhysics = MovementComp->UpdatedPrimitive && MovementComp->UpdatedPrimitive->IsSimulatingPhysics();
 
 		// Disable movement correction if needed
@@ -1155,19 +1160,39 @@ void UContextualAnimSceneActorComponent::SetMovementState(const FContextualAnimS
 			}
 		}
 
-		// Prevent physics rotation. During the interaction we want to be fully root motion driven
-		MovementComp->bAllowPhysicsRotationDuringAnimRootMotion = false;
-		MovementComp->bUseControllerDesiredRotation = false;
-		MovementComp->bOrientRotationToMovement = false;
+		if (AnimTrack.bControlCharacterRotation)
+		{
+			CharacterPropertiesBackup.RotationProperties = FCharacterRotationProperties();
+			CharacterPropertiesBackup.RotationProperties->bAllowPhysicsRotationDuringAnimRootMotion = MovementComp->bAllowPhysicsRotationDuringAnimRootMotion;
+			CharacterPropertiesBackup.RotationProperties->bUseControllerDesiredRotation = MovementComp->bUseControllerDesiredRotation;
+			CharacterPropertiesBackup.RotationProperties->bOrientRotationToMovement = MovementComp->bOrientRotationToMovement;
+			
+			// Prevent physics rotation. During the interaction we want to be fully root motion driven
+			MovementComp->bAllowPhysicsRotationDuringAnimRootMotion = false;
+			MovementComp->bUseControllerDesiredRotation = false;
+			MovementComp->bOrientRotationToMovement = false;
+		}
+		else
+		{
+			CharacterPropertiesBackup.RotationProperties.Reset();
+		}
 
 		if (UPrimitiveComponent* const UpdatedPrimitive = MovementComp->UpdatedPrimitive)
 		{
 			UpdatedPrimitive->SetSimulatePhysics(false);
 		}
 
-		if (MovementComp->MovementMode != DesiredMoveMode)
+		if (AnimTrack.bChangeMovementMode)
+		{
+			CharacterPropertiesBackup.MovementMode = MovementComp->MovementMode;
+			if (MovementComp->MovementMode != AnimTrack.MovementMode)
+			{
+				MovementComp->SetMovementMode(AnimTrack.MovementMode, AnimTrack.CustomMovementMode);
+			}
+		}
+		else
 		{
-			MovementComp->SetMovementMode(DesiredMoveMode, CustomMode);
+			CharacterPropertiesBackup.MovementMode.Reset();
 		}
 	}
 }
@@ -1186,22 +1211,30 @@ void UContextualAnimSceneActorComponent::RestoreMovementState(const FContextualA
 			}
 		}
 
-		MovementComp->bAllowPhysicsRotationDuringAnimRootMotion = CharacterPropertiesBackup.bAllowPhysicsRotationDuringAnimRootMotion;
-		MovementComp->bUseControllerDesiredRotation = CharacterPropertiesBackup.bUseControllerDesiredRotation;
-		MovementComp->bOrientRotationToMovement = CharacterPropertiesBackup.bOrientRotationToMovement;
-		if (
-			(CharacterPropertiesBackup.MovementMode == MOVE_Walking || CharacterPropertiesBackup.MovementMode == MOVE_NavWalking) &&
-			(MovementComp->GetGroundMovementMode() == MOVE_Walking || MovementComp->GetGroundMovementMode() == MOVE_NavWalking))
+		if (CharacterPropertiesBackup.RotationProperties.IsSet())
 		{
-			// This is mostly for AIs, we want to make sure to not override their ground movement mode if it changes while the anim is playing
-			// @TODO: We should have a more robust mechanism around this to prevent the same problem with any of the movement properties that change during the interaction.
-			MovementComp->SetMovementMode(MovementComp->GetGroundMovementMode());
+			MovementComp->bAllowPhysicsRotationDuringAnimRootMotion = CharacterPropertiesBackup.RotationProperties->bAllowPhysicsRotationDuringAnimRootMotion;
+			MovementComp->bUseControllerDesiredRotation = CharacterPropertiesBackup.RotationProperties->bUseControllerDesiredRotation;
+			MovementComp->bOrientRotationToMovement = CharacterPropertiesBackup.RotationProperties->bOrientRotationToMovement;
 		}
-		else
+		
+		if (CharacterPropertiesBackup.MovementMode.IsSet())
 		{
-			MovementComp->SetMovementMode(CharacterPropertiesBackup.MovementMode);
+			if (
+				(CharacterPropertiesBackup.MovementMode == MOVE_Walking || CharacterPropertiesBackup.MovementMode == MOVE_NavWalking) &&
+				(MovementComp->GetGroundMovementMode() == MOVE_Walking || MovementComp->GetGroundMovementMode() == MOVE_NavWalking))
+			{
+				// This is mostly for AIs, we want to make sure to not override their ground movement mode if it changes while the anim is playing
+				// @TODO: We should have a more robust mechanism around this to prevent the same problem with any of the movement properties that change during the interaction.
+				MovementComp->SetMovementMode(MovementComp->GetGroundMovementMode());
+			}
+			else
+			{
+				MovementComp->SetMovementMode(*CharacterPropertiesBackup.MovementMode);
+			}
 		}
-		if(UPrimitiveComponent* const UpdatedPrimitive = MovementComp->UpdatedPrimitive)
+		
+		if (UPrimitiveComponent* const UpdatedPrimitive = MovementComp->UpdatedPrimitive)
 		{
 			UpdatedPrimitive->SetSimulatePhysics(CharacterPropertiesBackup.bSimulatePhysics);
 		}
diff --git a/Engine/Plugins/Experimental/Animation/ContextualAnimation/Source/ContextualAnimation/Public/ContextualAnimSceneActorComponent.h b/Engine/Plugins/Experimental/Animation/ContextualAnimation/Source/ContextualAnimation/Public/ContextualAnimSceneActorComponent.h
--- a/Engine/Plugins/Experimental/Animation/ContextualAnimation/Source/ContextualAnimation/Public/ContextualAnimSceneActorComponent.h
+++ b/Engine/Plugins/Experimental/Animation/ContextualAnimation/Source/ContextualAnimation/Public/ContextualAnimSceneActorComponent.h
@@ -272,14 +272,19 @@ protected:
 	UPROPERTY(Transient)
 	TArray<FName> WarpTargetNamesCache;
 
-	struct FCharacterProperties
+	struct FCharacterRotationProperties
 	{
-		bool bIgnoreClientMovementErrorChecksAndCorrection = false;
 		bool bAllowPhysicsRotationDuringAnimRootMotion = false;
 		bool bUseControllerDesiredRotation = false;
 		bool bOrientRotationToMovement = false;
+	};
+
+	struct FCharacterProperties
+	{
+		bool bIgnoreClientMovementErrorChecksAndCorrection = false;
 		bool bSimulatePhysics = false;
-		EMovementMode MovementMode = EMovementMode::MOVE_Walking;
+		TOptional<FCharacterRotationProperties> RotationProperties;
+		TOptional<EMovementMode> MovementMode;
 		TArray<TTuple<ECollisionChannel, ECollisionResponse>> CollisionResponses;
 	};
 	FCharacterProperties CharacterPropertiesBackup;
@@ -311,7 +316,7 @@ protected:
 	virtual void SetCollisionState(const FContextualAnimSceneBinding& Binding);
 	virtual void RestoreCollisionState(const FContextualAnimSceneBinding& Binding);
 
-	void SetMovementState(const FContextualAnimSceneBinding& Binding, EMovementMode DesiredMoveMode, uint8 CustomMode = 0);
+	void SetMovementState(const FContextualAnimSceneBinding& Binding, const FContextualAnimTrack& AnimTrack);
 	void RestoreMovementState(const FContextualAnimSceneBinding& Binding);
 
 	UFUNCTION()
diff --git a/Engine/Plugins/Experimental/Animation/ContextualAnimation/Source/ContextualAnimation/Public/ContextualAnimTypes.h b/Engine/Plugins/Experimental/Animation/ContextualAnimation/Source/ContextualAnimation/Public/ContextualAnimTypes.h
--- a/Engine/Plugins/Experimental/Animation/ContextualAnimation/Source/ContextualAnimation/Public/ContextualAnimTypes.h
+++ b/Engine/Plugins/Experimental/Animation/ContextualAnimation/Source/ContextualAnimation/Public/ContextualAnimTypes.h
@@ -74,12 +74,22 @@ struct CONTEXTUALANIMATION_API FContextualAnimTrack
 	UPROPERTY(EditAnywhere, BlueprintReadWrite, Category = "Defaults")
 	float AnimMaxStartTime = 0.f;
 
+	/** If true, will set the movement mode of the character using CAS, to 'MovementMode', during the scope of this anim track. */
 	UPROPERTY(EditAnywhere, BlueprintReadWrite, Category = "Defaults")
+	bool bChangeMovementMode = true;
+
+	/** If 'bChangeMovementMode' is true, will set the movement mode of the character using CAS, to this value, during the scope of this anim track. */
+	UPROPERTY(EditAnywhere, BlueprintReadWrite, Category = "Defaults", meta = (EditConditionHides, EditCondition = "bChangeMovementMode"))
 	TEnumAsByte<enum EMovementMode> MovementMode = EMovementMode::MOVE_Walking;
 
+	/** If 'bChangeMovementMode' is true, and the 'MovementMode' is Custom, it'll set this custom movement mode during the scope of this anim track. */
 	UPROPERTY(EditAnywhere, BlueprintReadWrite, Category = "Defaults", meta = (EditConditionHides, EditCondition = "MovementMode==EMovementMode::MOVE_Custom"))
 	uint8 CustomMovementMode = 0;
 
+	/** If true: disable orientation of the character towards movement, controller's rotation and physics driven rotation during root motion. */
+	UPROPERTY(EditAnywhere, BlueprintReadWrite, Category = "Defaults")
+	bool bControlCharacterRotation = true;
+
 	/** Whether the actor that should play this animation is optional */
 	UPROPERTY(EditAnywhere, BlueprintReadWrite, Category = "Defaults")
 	bool bOptional = false;
