Contextual Anim Editor - Potential fix for rare crash in the contextual anim editor.

diff --git a/Engine/Plugins/Experimental/Animation/ContextualAnimation/Source/ContextualAnimationEditor/Private/ContextualAnimAssetEditorToolkit.cpp b/Engine/Plugins/Experimental/Animation/ContextualAnimation/Source/ContextualAnimationEditor/Private/ContextualAnimAssetEditorToolkit.cpp
--- a/Engine/Plugins/Experimental/Animation/ContextualAnimation/Source/ContextualAnimationEditor/Private/ContextualAnimAssetEditorToolkit.cpp
+++ b/Engine/Plugins/Experimental/Animation/ContextualAnimation/Source/ContextualAnimationEditor/Private/ContextualAnimAssetEditorToolkit.cpp
@@ -353,11 +353,13 @@ TSharedRef<SDockTab> FContextualAnimAssetEditorToolkit::SpawnTab_Timeline(const
 {
 	check(Args.GetTabId().TabType == FContextualAnimEditorTabs::TimelineID);
 
-	TSharedRef<SDockTab> SpawnedTab =
-		SNew(SDockTab)
-		[
-			ViewModel->GetSequencer()->GetSequencerWidget()
-		];
+	TSharedRef<SDockTab> SpawnedTab = SNew(SDockTab).Label(LOCTEXT("TimelineTab_Title", "Timeline"));
+
+	TSharedPtr<ISequencer> SequencerPtr = ViewModel ? ViewModel->GetSequencer() : nullptr;
+	if (ViewModel && ViewModel->GetSequencer() != nullptr)
+	{
+		SpawnedTab->SetContent(SequencerPtr->GetSequencerWidget());
+	}
 
 	return SpawnedTab;
 }
diff --git a/Engine/Plugins/Experimental/Animation/ContextualAnimation/Source/ContextualAnimationEditor/Private/ContextualAnimMovieSceneNotifyTrackEditor.cpp b/Engine/Plugins/Experimental/Animation/ContextualAnimation/Source/ContextualAnimationEditor/Private/ContextualAnimMovieSceneNotifyTrackEditor.cpp
--- a/Engine/Plugins/Experimental/Animation/ContextualAnimation/Source/ContextualAnimationEditor/Private/ContextualAnimMovieSceneNotifyTrackEditor.cpp
+++ b/Engine/Plugins/Experimental/Animation/ContextualAnimation/Source/ContextualAnimationEditor/Private/ContextualAnimMovieSceneNotifyTrackEditor.cpp
@@ -312,6 +312,12 @@ void FContextualAnimMovieSceneNotifyTrackEditor::BuildObjectBindingTrackMenu(FMe
 
 void FContextualAnimMovieSceneNotifyTrackEditor::AddNewNotifyTrack()
 {
+	TSharedPtr<ISequencer> SequencerPtr = GetSequencer();
+	if (!SequencerPtr.IsValid())
+	{
+		return;
+	}
+
 	// Copied from AnimTimelineTrack_Notifies.cpp/FAnimTimelineTrack_Notifies::GetNewTrackName()
 	auto GetNewTrackName = [](UAnimSequenceBase* InAnimSequenceBase) -> FName
 	{
@@ -353,7 +359,7 @@ void FContextualAnimMovieSceneNotifyTrackEditor::AddNewNotifyTrack()
 
 	MovieSceneTrack->Initialize(*Animation, NewNotifyTrack);
 
-	GetSequencer()->NotifyMovieSceneDataChanged(EMovieSceneDataChangeType::MovieSceneStructureItemAdded);
+	SequencerPtr->NotifyMovieSceneDataChanged(EMovieSceneDataChangeType::MovieSceneStructureItemAdded);
 }
 
 bool FContextualAnimMovieSceneNotifyTrackEditor::SupportsType(TSubclassOf<UMovieSceneTrack> Type) const
