Reverting the custom mode to the previous value when changing the movement mode in CAS

diff --git a/Engine/Plugins/Experimental/Animation/ContextualAnimation/Source/ContextualAnimation/Private/ContextualAnimSceneActorComponent.cpp b/Engine/Plugins/Experimental/Animation/ContextualAnimation/Source/ContextualAnimation/Private/ContextualAnimSceneActorComponent.cpp
--- a/Engine/Plugins/Experimental/Animation/ContextualAnimation/Source/ContextualAnimation/Private/ContextualAnimSceneActorComponent.cpp
+++ b/Engine/Plugins/Experimental/Animation/ContextualAnimation/Source/ContextualAnimation/Private/ContextualAnimSceneActorComponent.cpp
@@ -1187,6 +1187,7 @@ void UContextualAnimSceneActorComponent::SetMovementState(const FContextualAnimS
 		if (AnimTrack.bChangeMovementMode)
 		{
 			CharacterPropertiesBackup.MovementMode = MovementComp->MovementMode;
+			CharacterPropertiesBackup.CustomMovementMode = MovementComp->CustomMovementMode;
 			if (MovementComp->MovementMode != AnimTrack.MovementMode)
 			{
 				MovementComp->SetMovementMode(AnimTrack.MovementMode, AnimTrack.CustomMovementMode);
@@ -1195,6 +1196,7 @@ void UContextualAnimSceneActorComponent::SetMovementState(const FContextualAnimS
 		else
 		{
 			CharacterPropertiesBackup.MovementMode.Reset();
+			CharacterPropertiesBackup.CustomMovementMode = 0;
 		}
 	}
 }
@@ -1228,11 +1230,11 @@ void UContextualAnimSceneActorComponent::RestoreMovementState(const FContextualA
 			{
 				// This is mostly for AIs, we want to make sure to not override their ground movement mode if it changes while the anim is playing
 				// @TODO: We should have a more robust mechanism around this to prevent the same problem with any of the movement properties that change during the interaction.
-				MovementComp->SetMovementMode(MovementComp->GetGroundMovementMode());
+				MovementComp->SetMovementMode(MovementComp->GetGroundMovementMode(), CharacterPropertiesBackup.CustomMovementMode);
 			}
-			else
+			else if (*CharacterPropertiesBackup.MovementMode)
 			{
-				MovementComp->SetMovementMode(*CharacterPropertiesBackup.MovementMode);
+				MovementComp->SetMovementMode(*CharacterPropertiesBackup.MovementMode, CharacterPropertiesBackup.CustomMovementMode);
 			}
 		}
 		
diff --git a/Engine/Plugins/Experimental/Animation/ContextualAnimation/Source/ContextualAnimation/Public/ContextualAnimSceneActorComponent.h b/Engine/Plugins/Experimental/Animation/ContextualAnimation/Source/ContextualAnimation/Public/ContextualAnimSceneActorComponent.h
--- a/Engine/Plugins/Experimental/Animation/ContextualAnimation/Source/ContextualAnimation/Public/ContextualAnimSceneActorComponent.h
+++ b/Engine/Plugins/Experimental/Animation/ContextualAnimation/Source/ContextualAnimation/Public/ContextualAnimSceneActorComponent.h
@@ -285,6 +285,7 @@ protected:
 		bool bSimulatePhysics = false;
 		TOptional<FCharacterRotationProperties> RotationProperties;
 		TOptional<EMovementMode> MovementMode;
+		uint8 CustomMovementMode = 0;
 		TArray<TTuple<ECollisionChannel, ECollisionResponse>> CollisionResponses;
 	};
 	FCharacterProperties CharacterPropertiesBackup;
