Contextual Anim - Added GetAttachmentParams function in SceneAsset

diff --git a/Engine/Plugins/Experimental/Animation/ContextualAnimation/Source/ContextualAnimation/Public/ContextualAnimSceneAsset.h b/Engine/Plugins/Experimental/Animation/ContextualAnimation/Source/ContextualAnimation/Public/ContextualAnimSceneAsset.h
--- a/Engine/Plugins/Experimental/Animation/ContextualAnimation/Source/ContextualAnimation/Public/ContextualAnimSceneAsset.h
+++ b/Engine/Plugins/Experimental/Animation/ContextualAnimation/Source/ContextualAnimation/Public/ContextualAnimSceneAsset.h
@@ -241,6 +241,8 @@ public:
 
 	const TArray<TEnumAsByte<ECollisionChannel>>& GetCollisionChannelsToIgnoreForRole(FName Role) const;
 	
+	const TArray<FContextualAnimAttachmentParams>& GetAttachmentParams() const { return AttachmentParams; }
+
 	const FContextualAnimAttachmentParams* GetAttachmentParamsForRole(FName Role) const
 	{
 		return AttachmentParams.FindByPredicate([Role](const FContextualAnimAttachmentParams& Item) { return Item.Role == Role; });
