Contextual Anim: Added option to disable movement replication for simulated proxy

diff --git a/Engine/Plugins/Experimental/Animation/ContextualAnimation/Source/ContextualAnimation/Private/ContextualAnimSceneActorComponent.cpp b/Engine/Plugins/Experimental/Animation/ContextualAnimation/Source/ContextualAnimation/Private/ContextualAnimSceneActorComponent.cpp
--- a/Engine/Plugins/Experimental/Animation/ContextualAnimation/Source/ContextualAnimation/Private/ContextualAnimSceneActorComponent.cpp
+++ b/Engine/Plugins/Experimental/Animation/ContextualAnimation/Source/ContextualAnimation/Private/ContextualAnimSceneActorComponent.cpp
@@ -1160,6 +1160,16 @@ void UContextualAnimSceneActorComponent::SetMovementState(const FContextualAnimS
 				CharacterPropertiesBackup.bIgnoreClientMovementErrorChecksAndCorrection = MovementComp->bIgnoreClientMovementErrorChecksAndCorrection;
 				MovementComp->bIgnoreClientMovementErrorChecksAndCorrection = true;
 			}
+
+			if (Asset->ShouldDisableMovementReplicationForSimulatedProxy())
+			{
+				AActor* OwnerActor = GetOwner();
+				if (OwnerActor && OwnerActor->HasAuthority() && OwnerActor->GetRemoteRole() == ROLE_SimulatedProxy)
+				{
+					CharacterPropertiesBackup.bReplicatingMovement = OwnerActor->IsReplicatingMovement();
+					OwnerActor->SetReplicateMovement(false);
+				}
+			}
 		}
 
 		if (AnimTrack.bControlCharacterRotation)
@@ -1213,6 +1223,15 @@ void UContextualAnimSceneActorComponent::RestoreMovementState(const FContextualA
 			{
 				MovementComp->bIgnoreClientMovementErrorChecksAndCorrection = CharacterPropertiesBackup.bIgnoreClientMovementErrorChecksAndCorrection;
 			}
+
+			if (Asset->ShouldDisableMovementReplicationForSimulatedProxy())
+			{
+				AActor* OwnerActor = GetOwner();
+				if (OwnerActor && OwnerActor->HasAuthority() && OwnerActor->GetRemoteRole() == ROLE_SimulatedProxy)
+				{
+					OwnerActor->SetReplicateMovement(CharacterPropertiesBackup.bReplicatingMovement);
+				}
+			}
 		}
 
 		if (CharacterPropertiesBackup.RotationProperties.IsSet())
diff --git a/Engine/Plugins/Experimental/Animation/ContextualAnimation/Source/ContextualAnimation/Public/ContextualAnimSceneActorComponent.h b/Engine/Plugins/Experimental/Animation/ContextualAnimation/Source/ContextualAnimation/Public/ContextualAnimSceneActorComponent.h
--- a/Engine/Plugins/Experimental/Animation/ContextualAnimation/Source/ContextualAnimation/Public/ContextualAnimSceneActorComponent.h
+++ b/Engine/Plugins/Experimental/Animation/ContextualAnimation/Source/ContextualAnimation/Public/ContextualAnimSceneActorComponent.h
@@ -282,6 +282,7 @@ protected:
 	struct FCharacterProperties
 	{
 		bool bIgnoreClientMovementErrorChecksAndCorrection = false;
+		bool bReplicatingMovement = false;
 		bool bSimulatePhysics = false;
 		TOptional<FCharacterRotationProperties> RotationProperties;
 		TOptional<EMovementMode> MovementMode;
diff --git a/Engine/Plugins/Experimental/Animation/ContextualAnimation/Source/ContextualAnimation/Public/ContextualAnimSceneAsset.h b/Engine/Plugins/Experimental/Animation/ContextualAnimation/Source/ContextualAnimation/Public/ContextualAnimSceneAsset.h
--- a/Engine/Plugins/Experimental/Animation/ContextualAnimation/Source/ContextualAnimation/Public/ContextualAnimSceneAsset.h
+++ b/Engine/Plugins/Experimental/Animation/ContextualAnimation/Source/ContextualAnimation/Public/ContextualAnimSceneAsset.h
@@ -238,6 +238,7 @@ public:
 	FORCEINLINE float GetRadius() const { return Radius; }
 	FORCEINLINE bool ShouldPrecomputeAlignmentTracks() const { return bPrecomputeAlignmentTracks; }
 	FORCEINLINE bool ShouldIgnoreClientMovementErrorChecksAndCorrection() const { return bIgnoreClientMovementErrorChecksAndCorrection; }
+	FORCEINLINE bool ShouldDisableMovementReplicationForSimulatedProxy() const { return bDisableMovementReplicationForSimulatedProxy; }
 
 	const TArray<TEnumAsByte<ECollisionChannel>>& GetCollisionChannelsToIgnoreForRole(FName Role) const;
 	
@@ -359,9 +360,14 @@ protected:
 	UPROPERTY(EditAnywhere, Category = "Settings", meta = (DisplayName = "IK Target Params"))
 	FContextualAnimIKTargetParams IKTargetParams;
 
+	/** Whether to ignore movement error checks and corrections during the interaction for player characters */
 	UPROPERTY(EditAnywhere, Category = "Settings", AdvancedDisplay)
 	bool bIgnoreClientMovementErrorChecksAndCorrection = true;
 
+	/** Whether to disable movement replication during the interaction for simulated proxies (NPCs only). */
+	UPROPERTY(EditAnywhere, Category = "Settings", AdvancedDisplay)
+	bool bDisableMovementReplicationForSimulatedProxy = false;
+
 	/** Whether we should extract and cache alignment tracks off line. */
 	UPROPERTY(EditAnywhere, Category = "Settings", AdvancedDisplay)
 	bool bPrecomputeAlignmentTracks = false;
