Fix GC leaks / "Reachable garbage object" warnings caused by actors participating in ContextualAnimScene:
- Pawns that join a CAS will trigger a call to UPrimitiveComponent::IgnoreActorWhenMoving with any other actor in the CAS (e.g. the SmartObject actor they are interacting with, or other pawns using the same SO)
- IgnoreActorWhenMoving adds the passed-in actor to an array of hard-references in the PrimitiveComponent.  However, it only removes the other actor from that array when the Pawn explicitly leaves the CAS at the end of the interaction.
- So if the SO is streamed out underneath of the Pawn, the Pawn is still holding a hard reference to it which blocks the GC cleanup and causes a leak. This also triggers the "Reachable garbage object" warning as well as a hitch in the GC.
- With this change ContextualAnimSceneActorComponent now registers for the EndPlay delegate on any actor passed to IgnoreActorWhenMoving, so it removes the hard reference as soon as the ignored actor attempts to unload.

diff --git a/Engine/Plugins/Experimental/Animation/ContextualAnimation/Source/ContextualAnimation/Private/ContextualAnimSceneActorComponent.cpp b/Engine/Plugins/Experimental/Animation/ContextualAnimation/Source/ContextualAnimation/Private/ContextualAnimSceneActorComponent.cpp
--- a/Engine/Plugins/Experimental/Animation/ContextualAnimation/Source/ContextualAnimation/Private/ContextualAnimSceneActorComponent.cpp
+++ b/Engine/Plugins/Experimental/Animation/ContextualAnimation/Source/ContextualAnimation/Private/ContextualAnimSceneActorComponent.cpp
@@ -947,9 +947,39 @@ void UContextualAnimSceneActorComponent::SetIgnoreCollisionWithActor(AActor& Act
 	if (UPrimitiveComponent* RootPrimitiveComponent = Cast<UPrimitiveComponent>(GetOwner()->GetRootComponent()))
 	{
 		RootPrimitiveComponent->IgnoreActorWhenMoving(&Actor, bValue);
+
+		// MoveIgnoreActors is an array of hard references, so handle this actor receiving EndPlay (e.g. when it streams out)
+		//  so we explicitly remove its reference, which unblocks it from being cleaned up by garbage collection
+		if (bValue)
+		{
+			Actor.OnEndPlay.AddUniqueDynamic(this, &UContextualAnimSceneActorComponent::OnIgnoredActorEndPlay);
+		}
+		else
+		{
+			Actor.OnEndPlay.RemoveDynamic(this, &UContextualAnimSceneActorComponent::OnIgnoredActorEndPlay);
+		}
 	}
 }
 
+void UContextualAnimSceneActorComponent::OnIgnoredActorEndPlay(AActor* InEndPlayActor, EEndPlayReason::Type EndPlayReason)
+{
+	AActor* OwnerActor = GetOwner();
+	if (!OwnerActor)
+	{
+		return;
+	}
+
+	UPrimitiveComponent* RootPrimitiveComponent = Cast<UPrimitiveComponent>(OwnerActor->GetRootComponent());
+	if (!RootPrimitiveComponent)
+	{
+		return;
+	}
+
+	// If we were ignoring an Actor and it was removed from Play, remove it from our MoveIgnoreActors list immediately 
+	//  so that this hard reference to it doesn't block garbage collection, causing leaks
+	RootPrimitiveComponent->IgnoreActorWhenMoving(InEndPlayActor, false);
+}
+
 void UContextualAnimSceneActorComponent::SetCollisionState(const FContextualAnimSceneBinding& Binding)
 {
 	if (const UContextualAnimSceneAsset* Asset = Bindings.GetSceneAsset())
diff --git a/Engine/Plugins/Experimental/Animation/ContextualAnimation/Source/ContextualAnimation/Public/ContextualAnimSceneActorComponent.h b/Engine/Plugins/Experimental/Animation/ContextualAnimation/Source/ContextualAnimation/Public/ContextualAnimSceneActorComponent.h
--- a/Engine/Plugins/Experimental/Animation/ContextualAnimation/Source/ContextualAnimation/Public/ContextualAnimSceneActorComponent.h
+++ b/Engine/Plugins/Experimental/Animation/ContextualAnimation/Source/ContextualAnimation/Public/ContextualAnimSceneActorComponent.h
@@ -232,6 +232,9 @@ public:
 
 	bool IsOwnerLocallyControlled() const;
 
+	UFUNCTION()
+	void OnIgnoredActorEndPlay(AActor* Actor, EEndPlayReason::Type EndPlayReason);
+
 protected:
 
 	/** Cache of the owner's AnimInstance for easy access. Only valid while an interaction is active */
