Contextual Anim: Replaced TObjectPtr with TWeakObjectPtr in FContextualAnimRepLateJoinData in an attempt to fix a rare crash on some platforms during the replication of a 'late join' apparently caused by the replicated Actor being already destroyed

diff --git a/Engine/Plugins/Experimental/Animation/ContextualAnimation/Source/ContextualAnimation/Private/ContextualAnimSceneActorComponent.cpp b/Engine/Plugins/Experimental/Animation/ContextualAnimation/Source/ContextualAnimation/Private/ContextualAnimSceneActorComponent.cpp
--- a/Engine/Plugins/Experimental/Animation/ContextualAnimation/Source/ContextualAnimation/Private/ContextualAnimSceneActorComponent.cpp
+++ b/Engine/Plugins/Experimental/Animation/ContextualAnimation/Source/ContextualAnimation/Private/ContextualAnimSceneActorComponent.cpp
@@ -364,7 +364,8 @@ void UContextualAnimSceneActorComponent::OnRep_LateJoinData()
 	// This is received by the leader of the interaction on every remote client
 
 	UE_LOG(LogContextualAnim, Verbose, TEXT("%-21s UContextualAnimSceneActorComponent::OnRep_LateJoinData Owner: %s Bindings Id: %d Section: %d Asset: %s. Requester: %s Role: %s RepCounter: %d"),
-		*UEnum::GetValueAsString(TEXT("Engine.ENetRole"), GetOwner()->GetLocalRole()), *GetNameSafe(GetOwner()), Bindings.GetID(), Bindings.GetSectionIdx(), *GetNameSafe(Bindings.GetSceneAsset()), *GetNameSafe(RepLateJoinData.Actor), *RepLateJoinData.Role.ToString(), RepLateJoinData.RepCounter);
+		*UEnum::GetValueAsString(TEXT("Engine.ENetRole"), GetOwner()->GetLocalRole()),
+		*GetNameSafe(GetOwner()), Bindings.GetID(), Bindings.GetSectionIdx(), *GetNameSafe(Bindings.GetSceneAsset()), *GetNameSafe(RepLateJoinData.Actor.Get()), *RepLateJoinData.Role.ToString(), RepLateJoinData.RepCounter);
 
 	if (!RepLateJoinData.IsValid())
 	{
@@ -377,8 +378,13 @@ void UContextualAnimSceneActorComponent::OnRep_LateJoinData()
 		return;
 	}
 
+	AActor* Actor = RepLateJoinData.Actor.Get();
+	if (Actor == nullptr)
+	{
+		UE_LOG(LogContextualAnim, Warning, TEXT("%-21s UContextualAnimSceneActorComponent::OnRep_LateJoinData Invalid Actor in RepLateJoinData"), *UEnum::GetValueAsString(TEXT("Engine.ENetRole"), GetOwner()->GetLocalRole()));
+		return;
+	}
 
-	AActor* Actor = RepLateJoinData.Actor;
 	FName Role = RepLateJoinData.Role;
 
 	// Add actor to the bindings
diff --git a/Engine/Plugins/Experimental/Animation/ContextualAnimation/Source/ContextualAnimation/Public/ContextualAnimSceneActorComponent.h b/Engine/Plugins/Experimental/Animation/ContextualAnimation/Source/ContextualAnimation/Public/ContextualAnimSceneActorComponent.h
--- a/Engine/Plugins/Experimental/Animation/ContextualAnimation/Source/ContextualAnimation/Public/ContextualAnimSceneActorComponent.h
+++ b/Engine/Plugins/Experimental/Animation/ContextualAnimation/Source/ContextualAnimation/Public/ContextualAnimSceneActorComponent.h
@@ -90,7 +90,7 @@ struct FContextualAnimRepLateJoinData : public FContextualAnimRepData
 
 	/** Actor that is joining the interaction */
 	UPROPERTY()
-	TObjectPtr<AActor> Actor = nullptr;
+	TWeakObjectPtr<AActor> Actor = nullptr;
 
 	/** Role in the interaction the actor is gonna play */
 	UPROPERTY()
