Contextual Anim - Changed parent class for ContextualAnimSceneActorComp from Primitive Component to Scene Component for performance reasons. It was a primitive component initially only to get some rendering feature for debugging purpose that is not used anymore.

diff --git a/Engine/Plugins/Experimental/Animation/ContextualAnimation/Source/ContextualAnimation/Private/ContextualAnimSceneActorComponent.cpp b/Engine/Plugins/Experimental/Animation/ContextualAnimation/Source/ContextualAnimation/Private/ContextualAnimSceneActorComponent.cpp
--- a/Engine/Plugins/Experimental/Animation/ContextualAnimation/Source/ContextualAnimation/Private/ContextualAnimSceneActorComponent.cpp
+++ b/Engine/Plugins/Experimental/Animation/ContextualAnimation/Source/ContextualAnimation/Private/ContextualAnimSceneActorComponent.cpp
@@ -77,9 +77,6 @@ UContextualAnimSceneActorComponent::UContextualAnimSceneActorComponent(const FOb
 	PrimaryComponentTick.bCanEverTick = false;
 	PrimaryComponentTick.bStartWithTickEnabled = false;
 	SetIsReplicatedByDefault(true);
-
-	SetCollisionEnabled(ECollisionEnabled::NoCollision);
-	SetGenerateOverlapEvents(false);
 }
 
 void UContextualAnimSceneActorComponent::AddReferencedObjects(UObject* InThis, FReferenceCollector& Collector)
@@ -927,13 +924,6 @@ void UContextualAnimSceneActorComponent::OnRep_Bindings()
 	}
 }
 
-FBoxSphereBounds UContextualAnimSceneActorComponent::CalcBounds(const FTransform& LocalToWorld) const
-{
-	// The option of having an SceneAsset and draw options on this component may go away in the future anyway, replaced by smart objects.
-	const float Radius = SceneAsset && SceneAsset->HasValidData() ? SceneAsset->GetRadius() : 0.f;
-	return FBoxSphereBounds(FSphere(GetComponentTransform().GetLocation(), Radius));
-}
-
 void UContextualAnimSceneActorComponent::SetIgnoreCollisionWithOtherActors(bool bValue) const
 {
 	const AActor* OwnerActor = GetOwner();
@@ -1531,120 +1521,3 @@ const FContextualAnimIKTarget& UContextualAnimSceneActorComponent::GetIKTargetBy
 
 	return IKTargetPtr ? *IKTargetPtr : FContextualAnimIKTarget::InvalidIKTarget;
 }
-
-FPrimitiveSceneProxy* UContextualAnimSceneActorComponent::CreateSceneProxy()
-{
-	class FSceneActorCompProxy final : public FPrimitiveSceneProxy
-	{
-	public:
-
-		SIZE_T GetTypeHash() const override
-		{
-			static size_t UniquePointer;
-			return reinterpret_cast<size_t>(&UniquePointer);
-		}
-
-		FSceneActorCompProxy(const UContextualAnimSceneActorComponent* InComponent)
-			: FPrimitiveSceneProxy(InComponent)
-			, SceneAssetPtr(InComponent->SceneAsset)
-		{
-		}
-
-		virtual void GetDynamicMeshElements(const TArray<const FSceneView*>& Views, const FSceneViewFamily& ViewFamily, uint32 VisibilityMap, FMeshElementCollector& Collector) const override
-		{
-			const UContextualAnimSceneAsset* Asset = SceneAssetPtr.Get();
-			if (Asset == nullptr)
-			{
-				return;
-			}
-
-			const FMatrix& LocalToWorld = GetLocalToWorld();
-			const FTransform ToWorldTransform = FTransform(LocalToWorld);
-
-			for (int32 ViewIndex = 0; ViewIndex < Views.Num(); ViewIndex++)
-			{
-				if (VisibilityMap & (1 << ViewIndex))
-				{
-					const FSceneView* View = Views[ViewIndex];
-
-					// Taking into account the min and maximum drawing distance
-					const float DistanceSqr = (View->ViewMatrices.GetViewOrigin() - LocalToWorld.GetOrigin()).SizeSquared();
-					if (DistanceSqr < FMath::Square(GetMinDrawDistance()) || DistanceSqr > FMath::Square(GetMaxDrawDistance()))
-					{
-						continue;
-					}
-
-					FPrimitiveDrawInterface* PDI = Collector.GetPDI(ViewIndex);
-
-					//DrawCircle(PDI, ToWorldTransform.GetLocation(), FVector(1, 0, 0), FVector(0, 1, 0), FColor::Red, SceneAssetPtr->GetRadius(), 12, SDPG_World, 1.f);
-
-					SceneAssetPtr->ForEachAnimTrack([this, ToWorldTransform, PDI](const FContextualAnimTrack& AnimTrack)
-					{
-						if (AnimTrack.Role != SceneAssetPtr->GetPrimaryRole())
-						{
-							// Draw Entry Point
-							const FTransform EntryTransform = (SceneAssetPtr->GetAlignmentTransform(AnimTrack, 0, 0.f) * ToWorldTransform);
-							DrawCoordinateSystem(PDI, EntryTransform.GetLocation(), EntryTransform.Rotator(), 20.f, SDPG_World, 3.f);
-
-							// Draw Sync Point
-							const FTransform SyncPoint = SceneAssetPtr->GetAlignmentTransform(AnimTrack, 0, AnimTrack.GetSyncTimeForWarpSection(0)) * ToWorldTransform;
-							DrawCoordinateSystem(PDI, SyncPoint.GetLocation(), SyncPoint.Rotator(), 20.f, SDPG_World, 3.f);
-
-							FLinearColor DrawColor = FLinearColor::White;
-							for (const UContextualAnimSelectionCriterion* Criterion : AnimTrack.SelectionCriteria)
-							{
-								if (const UContextualAnimSelectionCriterion_TriggerArea* Spatial = Cast<UContextualAnimSelectionCriterion_TriggerArea>(Criterion))
-								{
-									const float HalfHeight = Spatial->Height / 2.f;
-									const int32 LastIndex = Spatial->PolygonPoints.Num() - 1;
-									for (int32 Idx = 0; Idx <= LastIndex; Idx++)
-									{
-										const FVector P0 = ToWorldTransform.TransformPositionNoScale(Spatial->PolygonPoints[Idx]);
-										const FVector P1 = ToWorldTransform.TransformPositionNoScale(Spatial->PolygonPoints[Idx == LastIndex ? 0 : Idx + 1]);
-
-										PDI->DrawLine(P0, P1, DrawColor, SDPG_Foreground, 2.f);
-										PDI->DrawLine(P0 + FVector::UpVector * Spatial->Height, P1 + FVector::UpVector * Spatial->Height, DrawColor, SDPG_Foreground, 2.f);
-
-										PDI->DrawLine(P0, P0 + FVector::UpVector * Spatial->Height, DrawColor, SDPG_Foreground, 2.f);
-									}
-								}
-							}
-						}
-
-						return UE::ContextualAnim::EForEachResult::Continue;
-					});
-				}
-			}
-		}
-
-		virtual FPrimitiveViewRelevance GetViewRelevance(const FSceneView* View) const override
-		{
-			const bool bShowForCollision = View->Family->EngineShowFlags.Collision;
-			FPrimitiveViewRelevance Result;
-			Result.bDrawRelevance = IsShown(View);
-			Result.bDynamicRelevance = true;
-			Result.bNormalTranslucency = Result.bSeparateTranslucency = IsShown(View);
-			return Result;
-		}
-
-		virtual uint32 GetMemoryFootprint(void) const override
-		{
-			return(sizeof(*this) + GetAllocatedSize());
-		}
-
-		uint32 GetAllocatedSize(void) const
-		{
-			return(FPrimitiveSceneProxy::GetAllocatedSize());
-		}
-
-	private:
-		TWeakObjectPtr<const UContextualAnimSceneAsset> SceneAssetPtr;
-	};
-
-	if(bEnableDebug)
-	{
-		return new FSceneActorCompProxy(this);
-	}
-
-	return nullptr;
-}
diff --git a/Engine/Plugins/Experimental/Animation/ContextualAnimation/Source/ContextualAnimation/Public/ContextualAnimSceneActorComponent.h b/Engine/Plugins/Experimental/Animation/ContextualAnimation/Source/ContextualAnimation/Public/ContextualAnimSceneActorComponent.h
--- a/Engine/Plugins/Experimental/Animation/ContextualAnimation/Source/ContextualAnimation/Public/ContextualAnimSceneActorComponent.h
+++ b/Engine/Plugins/Experimental/Animation/ContextualAnimation/Source/ContextualAnimation/Public/ContextualAnimSceneActorComponent.h
@@ -4,7 +4,7 @@
 
 #include "ContextualAnimTypes.h"
 #include "ActorComponents/IKRigInterface.h"
-#include "Components/PrimitiveComponent.h"
+#include "Components/SceneComponent.h"
 #include "ContextualAnimSceneActorComponent.generated.h"
 
 class AActor;
@@ -149,7 +149,7 @@ struct FContextualAnimRepTransitionData : public FContextualAnimRepData
 };
 
 UCLASS(meta = (BlueprintSpawnableComponent))
-class CONTEXTUALANIMATION_API UContextualAnimSceneActorComponent : public UPrimitiveComponent, public IIKGoalCreatorInterface
+class CONTEXTUALANIMATION_API UContextualAnimSceneActorComponent : public USceneComponent, public IIKGoalCreatorInterface
 {
 	GENERATED_BODY()
 
@@ -172,16 +172,11 @@ public:
 	UPROPERTY(EditAnywhere, BlueprintReadWrite, Category = "Config")
 	TObjectPtr<class UContextualAnimSceneAsset> SceneAsset;
 
-	UPROPERTY(EditAnywhere, BlueprintReadWrite, Category = "Debug")
-	bool bEnableDebug;
-
 	UContextualAnimSceneActorComponent(const FObjectInitializer& ObjectInitializer);
 
 	static void AddReferencedObjects(UObject* InThis, FReferenceCollector& Collector);
 	virtual void GetLifetimeReplicatedProps(TArray< FLifetimeProperty >& OutLifetimeProps) const;
 	virtual void EndPlay(const EEndPlayReason::Type EndPlayReason) override;
-	virtual FBoxSphereBounds CalcBounds(const FTransform& LocalToWorld) const;
-	virtual FPrimitiveSceneProxy* CreateSceneProxy() override;
 
 	virtual void AddIKGoals_Implementation(TMap<FName, FIKRigGoal>& OutGoals) override;
 
