Contextual Anim - Added an extra layer of randomness to the selection of the anim set. This helps to reduce the case where the same set was selected multiple times in a row.

diff --git a/Engine/Plugins/Experimental/Animation/ContextualAnimation/Source/ContextualAnimation/Private/ContextualAnimTypes.cpp b/Engine/Plugins/Experimental/Animation/ContextualAnimation/Source/ContextualAnimation/Private/ContextualAnimTypes.cpp
--- a/Engine/Plugins/Experimental/Animation/ContextualAnimation/Source/ContextualAnimation/Private/ContextualAnimTypes.cpp
+++ b/Engine/Plugins/Experimental/Animation/ContextualAnimation/Source/ContextualAnimation/Private/ContextualAnimTypes.cpp
@@ -13,6 +13,7 @@
 #include "Components/SkeletalMeshComponent.h"
 #include "GameFramework/CharacterMovementComponent.h"
 #include "MotionWarpingComponent.h"
+#include "Algo\RandomShuffle.h"
 
 #include UE_INLINE_GENERATED_CPP_BY_NAME(ContextualAnimTypes)
 
@@ -677,6 +678,9 @@ int32 FContextualAnimSceneBindings::FindAnimSet(const UContextualAnimSceneAsset&
 			}
 		}
 
+		// Shuffle the list of valid sets to add an extra layer of randomness. 
+		Algo::RandomShuffle(ValidSets);
+
 		float RandomValue = FMath::RandRange(0.f, TotalWeight);
 		for (int32 Idx = 0; Idx < ValidSets.Num(); Idx++)
 		{
