Contextual Anim - Converted ContextualAnimSceneActorComponent to ActorComponent for performance reasons. Added optional transform on the component to be able to offset the final transform used for alignment.

diff --git a/Engine/Plugins/Experimental/Animation/ContextualAnimation/Source/ContextualAnimation/Private/ContextualAnimTypes.cpp b/Engine/Plugins/Experimental/Animation/ContextualAnimation/Source/ContextualAnimation/Private/ContextualAnimTypes.cpp
--- a/Engine/Plugins/Experimental/Animation/ContextualAnimation/Source/ContextualAnimation/Private/ContextualAnimTypes.cpp
+++ b/Engine/Plugins/Experimental/Animation/ContextualAnimation/Source/ContextualAnimation/Private/ContextualAnimTypes.cpp
@@ -226,11 +226,12 @@ FTransform FContextualAnimSceneBindingContext::GetTransform() const
 	{
 		return ExternalTransform.GetValue();
 	}
-	// If no external transform is provided, use the transform of the scene actor comp
-	// We use this instead of the actor transform so we can have explicit control over the transform that represents the location/rotation of this actor during alignment
 	else if (UContextualAnimSceneActorComponent* Comp = GetSceneActorComponent())
 	{
-		return Comp->GetComponentTransform();
+		if (AActor* OwnerActor = Comp->GetOwner())
+		{
+			return Comp->AlignmentOffset * OwnerActor->GetActorTransform();
+		}
 	}
 
 	return FTransform::Identity;
diff --git a/Engine/Plugins/Experimental/Animation/ContextualAnimation/Source/ContextualAnimation/Public/ContextualAnimSceneActorComponent.h b/Engine/Plugins/Experimental/Animation/ContextualAnimation/Source/ContextualAnimation/Public/ContextualAnimSceneActorComponent.h
--- a/Engine/Plugins/Experimental/Animation/ContextualAnimation/Source/ContextualAnimation/Public/ContextualAnimSceneActorComponent.h
+++ b/Engine/Plugins/Experimental/Animation/ContextualAnimation/Source/ContextualAnimation/Public/ContextualAnimSceneActorComponent.h
@@ -4,7 +4,7 @@
 
 #include "ContextualAnimTypes.h"
 #include "ActorComponents/IKRigInterface.h"
-#include "Components/SceneComponent.h"
+#include "Components/ActorComponent.h"
 #include "ContextualAnimSceneActorComponent.generated.h"
 
 #define UE_API CONTEXTUALANIMATION_API
@@ -151,7 +151,7 @@ struct FContextualAnimRepTransitionData : public FContextualAnimRepData
 };
 
 UCLASS(MinimalAPI, meta = (BlueprintSpawnableComponent))
-class UContextualAnimSceneActorComponent : public USceneComponent, public IIKGoalCreatorInterface
+class UContextualAnimSceneActorComponent : public UActorComponent, public IIKGoalCreatorInterface
 {
 	GENERATED_BODY()
 
@@ -174,6 +174,10 @@ public:
 	UPROPERTY(EditAnywhere, BlueprintReadWrite, Category = "Config")
 	TObjectPtr<class UContextualAnimSceneAsset> SceneAsset;
 
+	/** Optional offset from the actor transform. Can be used to modify the final world space transform used for alignment. */
+	UPROPERTY(EditAnywhere, BlueprintReadWrite, Category = "Config")
+	FTransform AlignmentOffset;
+
 	UE_API UContextualAnimSceneActorComponent(const FObjectInitializer& ObjectInitializer);
 
 	static UE_API void AddReferencedObjects(UObject* InThis, FReferenceCollector& Collector);
diff --git a/Engine/Plugins/Experimental/Animation/ContextualAnimation/Source/ContextualAnimationEditor/Private/ContextualAnimViewModel.cpp b/Engine/Plugins/Experimental/Animation/ContextualAnimation/Source/ContextualAnimationEditor/Private/ContextualAnimViewModel.cpp
--- a/Engine/Plugins/Experimental/Animation/ContextualAnimation/Source/ContextualAnimationEditor/Private/ContextualAnimViewModel.cpp
+++ b/Engine/Plugins/Experimental/Animation/ContextualAnimation/Source/ContextualAnimationEditor/Private/ContextualAnimViewModel.cpp
@@ -296,7 +296,6 @@ AActor* FContextualAnimViewModel::SpawnPreviewActor(const FContextualAnimTrack&
 			UContextualAnimSceneActorComponent* SceneActorComp = NewObject<UContextualAnimSceneActorComponent>(PreviewCharacter);
 			SceneActorComp->RegisterComponentWithWorld(GetWorld());
 			SceneActorComp->InitializeComponent();
-			SceneActorComp->AttachToComponent(PreviewCharacter->GetRootComponent(), FAttachmentTransformRules(EAttachmentRule::SnapToTarget, false));
 
 			PreviewActor = PreviewCharacter;
 		}
@@ -347,7 +346,6 @@ AActor* FContextualAnimViewModel::SpawnPreviewActor(const FContextualAnimTrack&
 				UContextualAnimSceneActorComponent* SceneActorComp = NewObject<UContextualAnimSceneActorComponent>(PreviewActor);
 				SceneActorComp->RegisterComponentWithWorld(GetWorld());
 				SceneActorComp->InitializeComponent();
-				SceneActorComp->AttachToComponent(PreviewActor->GetRootComponent(), FAttachmentTransformRules(EAttachmentRule::SnapToTarget, false));
 			}
 		}
 
