Contextual Anim - Improved logging in FindAnimSet

diff --git a/Engine/Plugins/Experimental/Animation/ContextualAnimation/Source/ContextualAnimation/Private/ContextualAnimSelectionCriterion.cpp b/Engine/Plugins/Experimental/Animation/ContextualAnimation/Source/ContextualAnimation/Private/ContextualAnimSelectionCriterion.cpp
--- a/Engine/Plugins/Experimental/Animation/ContextualAnimation/Source/ContextualAnimation/Private/ContextualAnimSelectionCriterion.cpp
+++ b/Engine/Plugins/Experimental/Animation/ContextualAnimation/Source/ContextualAnimation/Private/ContextualAnimSelectionCriterion.cpp
@@ -72,9 +72,6 @@ bool UContextualAnimSelectionCriterion_TriggerArea::DoesQuerierPassCondition(con
 		bResult = (FMath::Abs(AngleSum) > 0.001f);
 	}
 
-	UE_LOG(LogContextualAnim, VeryVerbose, TEXT("UContextualAnimSelectionCriterion_TriggerArea: Primary: %s Querier: %s Result: %d"),
-		*GetNameSafe(Primary.GetActor()), *GetNameSafe(Querier.GetActor()), bResult);
-
 	return bResult;
 }
 
@@ -107,10 +104,6 @@ bool UContextualAnimSelectionCriterion_Cone::DoesQuerierPassCondition(const FCon
 			bResult = CheckCondition(QuerierTransform, PrimaryTransform);
 		}
 	}
-	
-
-	UE_LOG(LogContextualAnim, VeryVerbose, TEXT("UContextualAnimSelectionCriterion_Cone: Primary: %s Querier: %s Mode: %s Distance: %.1f HalfAngle: %.1f Offset: %.1f Result: %d"), 
-		*GetNameSafe(Primary.GetActor()), *GetNameSafe(Querier.GetActor()), *UEnum::GetValueAsString(TEXT("ContextualAnimation.EContextualAnimCriterionConeMode"), Mode), Distance, HalfAngle, Offset, bResult);
 
 	return bResult;
 }
@@ -134,9 +127,5 @@ bool UContextualAnimSelectionCriterion_Distance::DoesQuerierPassCondition(const
 	}
 
 	const bool bResult = FMath::IsWithinInclusive(Distance, MinDistance, MaxDistance);
-
-	UE_LOG(LogContextualAnim, VeryVerbose, TEXT("UContextualAnimSelectionCriterion_Distance: Primary: %s Querier: %s Mode: %s MaxDistance: %.1f MaxDist: %.1f Dist: %.1f Result: %d"),
-		*GetNameSafe(Primary.GetActor()), *GetNameSafe(Querier.GetActor()), *UEnum::GetValueAsString(TEXT("ContextualAnimation.EContextualAnimCriterionDistanceMode"), Mode), MinDistance, MaxDistance, Distance, bResult);
-
 	return bResult;
 }
diff --git a/Engine/Plugins/Experimental/Animation/ContextualAnimation/Source/ContextualAnimation/Private/ContextualAnimTypes.cpp b/Engine/Plugins/Experimental/Animation/ContextualAnimation/Source/ContextualAnimation/Private/ContextualAnimTypes.cpp
--- a/Engine/Plugins/Experimental/Animation/ContextualAnimation/Source/ContextualAnimation/Private/ContextualAnimTypes.cpp
+++ b/Engine/Plugins/Experimental/Animation/ContextualAnimation/Source/ContextualAnimation/Private/ContextualAnimTypes.cpp
@@ -173,8 +173,13 @@ bool FContextualAnimTrack::DoesQuerierPassSelectionCriteria(const FContextualAni
 	{
 		if (Criterion && !Criterion->DoesQuerierPassCondition(Primary, Querier))
 		{
+			UE_LOG(LogContextualAnim, Verbose, TEXT("\t\t\t '%s' condition check failed."), *GetNameSafe(Criterion));
 			return false;
 		}
+		else
+		{
+			UE_LOG(LogContextualAnim, Verbose, TEXT("\t\t\t '%s' condition check passed."), *GetNameSafe(Criterion));
+		}
 	}
 
 	return true;
@@ -526,24 +531,23 @@ bool FContextualAnimSceneBindings::CheckConditions(const UContextualAnimSceneAss
 
 	auto DoCheck = [&SceneAsset, SectionIdx, AnimSetIdx](const FContextualAnimSceneBindingContext& Primary, const FContextualAnimSceneBindingContext& Querier, const FName& Role)
 	{
-		if (AActor* Actor = Querier.GetActor())
+		AActor* Actor = Querier.GetActor();
+		UContextualAnimSceneActorComponent* SceneActorComp = Actor ? Actor->FindComponentByClass<UContextualAnimSceneActorComponent>() : nullptr;
+		if (SceneActorComp == nullptr)
 		{
-			UContextualAnimSceneActorComponent* SceneActorComp = Actor->FindComponentByClass<UContextualAnimSceneActorComponent>();
-			if (SceneActorComp == nullptr)
-			{
-				UE_LOG(LogContextualAnim, Verbose, TEXT("FContextualAnimSceneBindings::CheckConditions Failed. Reason: Missing ContextualAnimSceneActorComp. SceneAsset: %s Actor: %s Role: %s SectionIdx: %d AnimSetIdx: %d"),
-					*GetNameSafe(&SceneAsset), *GetNameSafe(Actor), *Role.ToString(), SectionIdx, AnimSetIdx);
-
-				return false;
-			}
+			UE_LOG(LogContextualAnim, Verbose, TEXT("\t\t\t Missing ContextualAnimSceneActorComponent on the Actor."));
+			return false;
 		}
-
+		
 		const FContextualAnimTrack* AnimTrack = SceneAsset.GetAnimTrack(SectionIdx, AnimSetIdx, Role);
-		if (AnimTrack == nullptr || AnimTrack->DoesQuerierPassSelectionCriteria(Primary, Querier) == false)
+		if (AnimTrack == nullptr)
 		{
-			UE_LOG(LogContextualAnim, Verbose, TEXT("FContextualAnimSceneBindings::CheckConditions Failed. Reason: Reason: Can't find valid track for actor. SceneAsset: %s Actor: %s Role: %s SectionIdx: %d AnimSetIdx: %d"),
-				*GetNameSafe(&SceneAsset), *GetNameSafe(Querier.GetActor()), *Role.ToString(), SectionIdx, AnimSetIdx);
+			UE_LOG(LogContextualAnim, Verbose, TEXT("\t\t\t Missing AnimTrack in SceneAsset for role '%s'."), *Role.ToString());
+			return false;
+		}
 
+		if (AnimTrack->DoesQuerierPassSelectionCriteria(Primary, Querier) == false)
+		{
 			return false;
 		}
 
@@ -555,18 +559,23 @@ bool FContextualAnimSceneBindings::CheckConditions(const UContextualAnimSceneAss
 	const FContextualAnimSceneBindingContext* PrimaryPtr = Params.Find(PrimaryRole);
 	if (PrimaryPtr == nullptr)
 	{
-		UE_LOG(LogContextualAnim, Verbose, TEXT("FContextualAnimSceneBindings::CheckConditions Failed. Reason: Can't find valid actor for primary role. SceneAsset: %s PrimaryRole: %s"),
-			*GetNameSafe(&SceneAsset), *PrimaryRole.ToString());
-
+		UE_LOG(LogContextualAnim, Verbose, TEXT("\t\t Missing actor to bind to the primary role '%s'."), *PrimaryRole.ToString());
 		return false;
 	}
 
 	// Test primary actor first
 	// Passing the same data twice (as primary and querier) feels weird, but this allow us to run the selection mechanism even on the primary actor.
+	UE_LOG(LogContextualAnim, Verbose, TEXT("\t\t Checking binding conditions for primary role '%s' (Actor: %s)"), *PrimaryRole.ToString(), *GetNameSafe(PrimaryPtr->GetActor()));
+
 	if (DoCheck(*PrimaryPtr, *PrimaryPtr, PrimaryRole) == false)
 	{
+		UE_LOG(LogContextualAnim, Verbose, TEXT("\t\t Actor '%s' failed to pass some of the conditions for primary role '%s'."), *GetNameSafe(PrimaryPtr->GetActor()), *PrimaryRole.ToString());
 		return false;
 	}
+	else
+	{
+		UE_LOG(LogContextualAnim, Verbose, TEXT("\t\t Actor '%s' passed all the conditions for primary role '%s'."), *GetNameSafe(PrimaryPtr->GetActor()), *PrimaryRole.ToString());
+	}
 
 	// Now test secondary actors
 	int32 ValidEntries = 1;
@@ -575,28 +584,29 @@ bool FContextualAnimSceneBindings::CheckConditions(const UContextualAnimSceneAss
 		FName RoleToBind = Pair.Key;
 		if (RoleToBind != PrimaryRole)
 		{
+			UE_LOG(LogContextualAnim, Verbose, TEXT("\t\t Checking binding conditions for role '%s' (Actor: %s)"), *RoleToBind.ToString(), *GetNameSafe(Pair.Value.GetActor()));
 			if (DoCheck(*PrimaryPtr, Pair.Value, RoleToBind) == false)
 			{
-				UE_LOG(LogContextualAnim, Verbose, TEXT("FContextualAnimSceneBindings::CheckConditions Failed. Reason: Can't find valid track for actor. SceneAsset: %s Actor: %s Role: %s SectionIdx: %d AnimSetIdx: %d"),
-					*GetNameSafe(&SceneAsset), *GetNameSafe(Pair.Value.GetActor()), *RoleToBind.ToString(), SectionIdx, AnimSetIdx);
-
+				UE_LOG(LogContextualAnim, Verbose, TEXT("\t\t Actor '%s' failed to pass some of the conditions for role '%s'."), *GetNameSafe(Pair.Value.GetActor()), *RoleToBind.ToString());
 				return false;
 			}
-
-			ValidEntries++;
+			else
+			{
+				UE_LOG(LogContextualAnim, Verbose, TEXT("\t\t Actor '%s' passed all the conditions for role '%s'."), *GetNameSafe(Pair.Value.GetActor()), *RoleToBind.ToString());	
+				ValidEntries++;
+			}
 		}
 	}
 
 	const int32 NumMandatoryRoles = SceneAsset.GetNumMandatoryRoles(SectionIdx, AnimSetIdx);
 	if (ValidEntries < NumMandatoryRoles)
 	{
-		UE_LOG(LogContextualAnim, Verbose, TEXT("FContextualAnimSceneBindings::CheckConditions Failed. Reason: Missing mandatory roles %d/%d. SceneAsset: %s SectionIdx: %d AnimSetIdx: %d"),
+		UE_LOG(LogContextualAnim, Verbose, TEXT("\t Missing mandatory roles %d/%d. SceneAsset: %s SectionIdx: %d AnimSetIdx: %d"),
 			ValidEntries, NumMandatoryRoles, *GetNameSafe(&SceneAsset), SectionIdx, AnimSetIdx);
 
 		return false;
 	}
 
-	UE_LOG(LogContextualAnim, Verbose, TEXT("FContextualAnimSceneBindings::CheckConditions Success! SceneAsset: %s SectionIdx: %d AnimSetIdx: %d"), *GetNameSafe(&SceneAsset), SectionIdx, AnimSetIdx);
 	return true;
 }
 
@@ -653,37 +663,68 @@ int32 FContextualAnimSceneBindings::FindAnimSet(const UContextualAnimSceneAsset&
 {
 	check(SceneAsset.HasValidData());
 
+	const FContextualAnimSceneSection* Section = SceneAsset.GetSection(SectionIdx);
+	if (Section == nullptr)
+	{
+		return INDEX_NONE;
+	}
+
 	int32 AnimSetIdxSelected = INDEX_NONE;
+	const int32 NumSets = Section->GetNumAnimSets();
+	int32 NumValid = 0;
+	
+	UE_LOG(LogContextualAnim, Verbose, TEXT("Find AnimSet in Section '%d' ('%s'). NumSets: %d. SceneAsset: %s"), 
+		SectionIdx, *Section->GetName().ToString(), NumSets, *GetNameSafe(&SceneAsset));
 
-	const int32 NumSets = SceneAsset.GetNumAnimSetsInSection(SectionIdx);
 	if (NumSets == 1)
 	{
+		const FContextualAnimSet* AnimSet = Section->GetAnimSet(0);
+		UE_LOG(LogContextualAnim, Verbose, TEXT("\t Checking binding conditions for AnimSet '0' ('%s')."), *AnimSet->Name.ToString());
+
 		if (CheckConditions(SceneAsset, SectionIdx, 0, Params))
 		{
 			AnimSetIdxSelected = 0;
+			NumValid = 1;
+
+			UE_LOG(LogContextualAnim, Verbose, TEXT("\t AnimSet '0' ('%s') can be selected"), *AnimSet->Name.ToString());
+		}
+		else
+		{
+			UE_LOG(LogContextualAnim, Verbose, TEXT("\t AnimSet '0' ('%s') can NOT be selected."), *AnimSet->Name.ToString());
 		}
 	}
 	else if (NumSets > 1)
 	{
-		const FContextualAnimSceneSection* Section = SceneAsset.GetSection(SectionIdx);
 		TArray<TTuple<int32, float>, TInlineAllocator<5>> ValidSets; // 0: AnimSetIdx, 1: AnimSetRandomWeight
 
 		float TotalWeight = 0;
 		for (int32 AnimSetIdx = 0; AnimSetIdx < NumSets; AnimSetIdx++)
 		{
+			const FContextualAnimSet* AnimSet = Section->GetAnimSet(AnimSetIdx);
+			UE_LOG(LogContextualAnim, Verbose, TEXT("\t Checking binding conditions for AnimSet '%d' ('%s')."), AnimSetIdx, *AnimSet->Name.ToString());
+
 			if (CheckConditions(SceneAsset, SectionIdx, AnimSetIdx, Params))
 			{
-				const float AnimSetRandomWeight = FMath::Max(Section->GetAnimSet(AnimSetIdx)->RandomWeight, 0);
+				const float AnimSetRandomWeight = FMath::Max(AnimSet->RandomWeight, 0);
 				ValidSets.Add(MakeTuple(AnimSetIdx, AnimSetRandomWeight));
 				TotalWeight += AnimSetRandomWeight;
+
+				UE_LOG(LogContextualAnim, Verbose, TEXT("\t AnimSet '%d' ('%s') can be selected. Adding it to the list of valid sets with a random weight of %.2f."), 
+					AnimSetIdx, *AnimSet->Name.ToString(), AnimSetRandomWeight);
+			}
+			else
+			{
+				UE_LOG(LogContextualAnim, Verbose, TEXT("\t AnimSet '%d' ('%s') can NOT be selected."), 
+					AnimSetIdx, *Section->GetAnimSet(AnimSetIdx)->Name.ToString());
 			}
 		}
 
 		// Shuffle the list of valid sets to add an extra layer of randomness. 
 		Algo::RandomShuffle(ValidSets);
 
+		NumValid = ValidSets.Num();
 		float RandomValue = FMath::RandRange(0.f, TotalWeight);
-		for (int32 Idx = 0; Idx < ValidSets.Num(); Idx++)
+		for (int32 Idx = 0; Idx < NumValid; Idx++)
 		{
 			RandomValue -= FMath::Max(ValidSets[Idx].Get<1>(), 0);
 			if (RandomValue <= 0)
@@ -694,7 +735,16 @@ int32 FContextualAnimSceneBindings::FindAnimSet(const UContextualAnimSceneAsset&
 		}
 	}
 
-	UE_CLOG(AnimSetIdxSelected == INDEX_NONE, LogContextualAnim, Warning, TEXT("FContextualAnimSceneBindings::FindAnimSet. Can't find AnimSet"));
+	if (AnimSetIdxSelected != INDEX_NONE)
+	{
+		UE_LOG(LogContextualAnim, Verbose, TEXT("Find AnimSet Result: SUCCESS. %d/%d AnimSets have passed the binding conditions. AnimSet '%d' ('%s') was selected."), 
+			NumValid, NumSets, AnimSetIdxSelected, *SceneAsset.GetSection(SectionIdx)->GetAnimSet(AnimSetIdxSelected)->Name.ToString());
+	}
+	else
+	{
+		UE_LOG(LogContextualAnim, Verbose, TEXT("Find AnimSet Result: FAILURE. No AnimSet passed the binding conditions."));
+	}
+
 	return AnimSetIdxSelected;
 }
 
