[Contextual Animation] [CAS] Add Optional Overrides to CAS Bindings
- Allow Users to Supply an Override Interface that can Swap out Animations at runtime for CAS Scenes. This will Allow us to create Custom Variations for Actors Sharing the Same Scene Setup
- Will Allow us to Publish new actors without having to update the Scenes themselves.
- Add Editor View Option to Visualize Overrides
- Ensures Overrides Do not try to use Pre-Computed Warp Points, Instead Recomputes Warp Points if Override Present at Runtime

diff --git a/Engine/Plugins/Experimental/Animation/ContextualAnimation/Source/ContextualAnimation/Private/ContextualAnimSceneActorComponent.cpp b/Engine/Plugins/Experimental/Animation/ContextualAnimation/Source/ContextualAnimation/Private/ContextualAnimSceneActorComponent.cpp
--- a/Engine/Plugins/Experimental/Animation/ContextualAnimation/Source/ContextualAnimation/Private/ContextualAnimSceneActorComponent.cpp
+++ b/Engine/Plugins/Experimental/Animation/ContextualAnimation/Source/ContextualAnimation/Private/ContextualAnimSceneActorComponent.cpp
@@ -158,7 +158,8 @@ void UContextualAnimSceneActorComponent::PlayAnimation_Internal(UAnimSequenceBas
 						if (SyncLeader->GetActor() != GetOwner())
 						{
 							FAnimMontageInstance* LeaderMontageInstance = SyncLeader->GetAnimMontageInstance();
-							if (LeaderMontageInstance && LeaderMontageInstance->Montage == Bindings.GetAnimTrackFromBinding(*SyncLeader).Animation && MontageInstance->GetMontageSyncLeader() == nullptr)
+							const UAnimSequenceBase* AnimSequence = SyncLeader->GetAnimation();
+							if (LeaderMontageInstance && LeaderMontageInstance->Montage == AnimSequence && MontageInstance->GetMontageSyncLeader() == nullptr)
 							{
 								UE_LOG(LogContextualAnim, VeryVerbose, TEXT("%-21s \t\tUContextualAnimSceneActorComponent::PlayAnimation_Internal Syncing Animation. Actor: %s Anim: %s StartTime: %f bSyncPlaybackTime: %d"),
 									*UEnum::GetValueAsString(TEXT("Engine.ENetRole"), GetOwner()->GetLocalRole()), *GetNameSafe(GetOwner()), *GetNameSafe(Animation), StartTime, bSyncPlaybackTime);
@@ -209,8 +210,8 @@ void UContextualAnimSceneActorComponent::AddOrUpdateWarpTargets(int32 SectionIdx
 				return;
 			}
 
-			const FContextualAnimTrack* AnimTrack = Asset->GetAnimTrack(SectionIdx, AnimSetIdx, Bindings.GetRoleFromBinding(*Binding));
-			if (AnimTrack == nullptr || AnimTrack->Animation == nullptr)
+			UAnimSequenceBase* AniamtionSequence = Binding->GetAnimation();
+			if (AniamtionSequence == nullptr)
 			{
 				return;
 			}
@@ -219,8 +220,9 @@ void UContextualAnimSceneActorComponent::AddOrUpdateWarpTargets(int32 SectionIdx
 			{
 				if (WarpPoint.Name != NAME_None)
 				{
-					const float Time = AnimTrack->GetSyncTimeForWarpSection(WarpPoint.Name);
-					const FTransform TransformRelativeToWarpPoint = Asset->GetAlignmentTransform(*AnimTrack, WarpPoint.Name, Time);
+					float StartTime, EndTime;
+					UContextualAnimUtilities::GetStartAndEndTimeForWarpSection(AniamtionSequence, WarpPoint.Name, StartTime, EndTime);
+					const FTransform TransformRelativeToWarpPoint = Asset->GetAlignmentTransform(*Binding, Bindings, WarpPoint.Name, EndTime);
 					const FTransform WarpTargetTransform = TransformRelativeToWarpPoint * WarpPoint.Transform;
 					MotionWarpComp->AddOrUpdateWarpTargetFromTransform(WarpPoint.Name, WarpTargetTransform);
 					WarpTargetNamesCache.AddUnique(WarpPoint.Name);
@@ -338,7 +340,10 @@ void UContextualAnimSceneActorComponent::LateJoinScene(const FContextualAnimScen
 		const FContextualAnimTrack* AnimTrack = Bindings.GetSceneAsset()->GetAnimTrack(SectionIdx, AnimSetIdx, Bindings.GetRoleFromBinding(*Binding));
 		check(AnimTrack);
 
-		PlayAnimation_Internal(AnimTrack->Animation, 0.f, false);
+
+		UAnimSequenceBase* AnimSequence = Binding->GetAnimation();
+
+		PlayAnimation_Internal(AnimSequence, 0.f, false);
 
 		AddOrUpdateWarpTargets(SectionIdx, AnimSetIdx, WarpPoints, ExternalWarpTargets);
 
@@ -531,11 +536,12 @@ void UContextualAnimSceneActorComponent::HandleTransitionSelf(int32 NewSectionId
 
 	// Play animation
 	const FContextualAnimSceneBinding& Binding = *Bindings.FindBindingByActor(GetOwner());
-	const FContextualAnimTrack& AnimTrack = Bindings.GetAnimTrackFromBinding(Binding);
-	PlayAnimation_Internal(AnimTrack.Animation, 0.f, Bindings.ShouldBindingSyncAnimation(Binding));
+	UAnimSequenceBase* AnimationSequence = Binding.GetAnimation();
+	PlayAnimation_Internal(AnimationSequence, 0.f, Bindings.ShouldBindingSyncAnimation(Binding));
 
 	AddOrUpdateWarpTargets(NewSectionIdx, NewAnimSetIdx, WarpPoints, ExternalWarpTargets);
 
+	const FContextualAnimTrack& AnimTrack = Bindings.GetAnimTrackFromBinding(Binding);
 	if (AnimTrack.bChangeMovementMode)
 	{
 		if (UCharacterMovementComponent* MovementComp = Binding.GetCharacterMovementComponent())
@@ -596,7 +602,8 @@ bool UContextualAnimSceneActorComponent::TransitionSingleActor(int32 SectionIdx,
 				TArray<FContextualAnimWarpPoint> WarpPoints;
 				CalculateWarpPointsForBindings(Bindings, SectionIdx, AnimSetIdx, WarpPoints);
 
-				PlayAnimation_Internal(AnimTrack->Animation, 0.f, false);
+				UAnimSequenceBase* AnimationSequence = OwnerBinding->GetAnimation();
+				PlayAnimation_Internal(AnimationSequence, 0.f, false);
 
 				AddOrUpdateWarpTargets(SectionIdx, AnimSetIdx, WarpPoints, ExternalWarpTargets);
 
@@ -656,7 +663,8 @@ void UContextualAnimSceneActorComponent::OnRep_RepTransitionSingleActor()
 				const FContextualAnimTrack* AnimTrack = Asset->GetAnimTrack(RepTransitionSingleActorData.SectionIdx, RepTransitionSingleActorData.AnimSetIdx, Bindings.GetRoleFromBinding(*OwnerBinding));
 				if (AnimTrack && AnimTrack->Animation)
 				{
-					PlayAnimation_Internal(AnimTrack->Animation, 0.f, false);
+					UAnimSequenceBase* TrackAnimationSequence = OwnerBinding->GetAnimation();
+					PlayAnimation_Internal(TrackAnimationSequence, 0.f, false);
 
 					AddOrUpdateWarpTargets(RepTransitionSingleActorData.SectionIdx, RepTransitionSingleActorData.AnimSetIdx, RepTransitionSingleActorData.WarpPoints, RepTransitionSingleActorData.ExternalWarpTargets);
 
@@ -802,7 +810,7 @@ void UContextualAnimSceneActorComponent::EarlyOutContextualAnimScene(bool bStopE
 		UE_LOG(LogContextualAnim, Verbose, TEXT("%-21s UContextualAnimSceneActorComponent::EarlyOutContextualAnimScene Actor: %s ActiveMontage: %s bStopEveryone: %d"),
 			*UEnum::GetValueAsString(TEXT("Engine.ENetRole"), GetOwner()->GetLocalRole()), *GetNameSafe(GetOwner()), *GetNameSafe(ActiveMontage), bStopEveryone);
 
-		if (Bindings.GetAnimTrackFromBinding(*Binding).Animation == ActiveMontage)
+		if (Binding->GetAnimation() == ActiveMontage)
 		{
 			const uint8 BindingsId = Bindings.GetID();
 
@@ -903,6 +911,9 @@ void UContextualAnimSceneActorComponent::OnRep_Bindings()
 	// The owner of this component started an interaction on the server
 	if (RepBindings.Bindings.IsValid())
 	{
+		// Locally Update the Binding Overrides, this may generate IK and ALignment Tracks if Needed
+		RepBindings.Bindings.UpdateBindingOverridesIfNeeded();
+
 		const FContextualAnimSceneBinding* OwnerBinding = RepBindings.Bindings.FindBindingByActor(GetOwner());
 		if (ensureAlways(OwnerBinding))
 		{
@@ -1059,13 +1070,14 @@ void UContextualAnimSceneActorComponent::JoinScene(const FContextualAnimSceneBin
 
 		Bindings = InBindings;
 
-		const FContextualAnimTrack& AnimTrack = Bindings.GetAnimTrackFromBinding(*Binding);
-		PlayAnimation_Internal(AnimTrack.Animation, 0.f, Bindings.ShouldBindingSyncAnimation(*Binding));
+		UAnimSequenceBase* AnimationSequence = Binding->GetAnimation();
+		PlayAnimation_Internal(AnimationSequence, 0.f, Bindings.ShouldBindingSyncAnimation(*Binding));
 
-		AddOrUpdateWarpTargets(AnimTrack.SectionIdx, AnimTrack.AnimSetIdx, WarpPoints, ExternalWarpTargets);
+		AddOrUpdateWarpTargets(Binding->GetSectionIdx(), Binding->GetAnimSetIdx(), WarpPoints, ExternalWarpTargets);
 
 		SetCollisionState(*Binding);
 
+		const FContextualAnimTrack& AnimTrack = Bindings.GetAnimTrackFromBinding(*Binding);
 		SetMovementState(*Binding, AnimTrack);
 
 		OnJoinScene(*Binding);
@@ -1433,7 +1445,8 @@ void UContextualAnimSceneActorComponent::UpdateIKTargets()
 		{
 			// Do not update if the target actor should be playing and animation but its not yet. 
 			// This could happen in multi player when the initiator start playing the animation locally
-			const UAnimSequenceBase* TargetAnimation = Bindings.GetAnimTrackFromBinding(*TargetBinding).Animation;
+
+			const UAnimSequenceBase* TargetAnimation = TargetBinding->GetAnimation();
 			if (TargetAnimation)
 			{
 				//@TODO: Add support for dynamic montages
diff --git a/Engine/Plugins/Experimental/Animation/ContextualAnimation/Source/ContextualAnimation/Private/ContextualAnimSceneAsset.cpp b/Engine/Plugins/Experimental/Animation/ContextualAnimation/Source/ContextualAnimation/Private/ContextualAnimSceneAsset.cpp
--- a/Engine/Plugins/Experimental/Animation/ContextualAnimation/Source/ContextualAnimation/Private/ContextualAnimSceneAsset.cpp
+++ b/Engine/Plugins/Experimental/Animation/ContextualAnimation/Source/ContextualAnimation/Private/ContextualAnimSceneAsset.cpp
@@ -658,6 +658,11 @@ FTransform UContextualAnimSceneAsset::GetIKTargetTransform(int32 SectionIdx, int
 	return FTransform::Identity;
 }
 
+FTransform UContextualAnimSceneAsset::GetIKTargetTransform(const FContextualAnimSceneBinding& InBinding, const FName& TrackName, float Time) const
+{
+	return InBinding.GetIKTargetData().ExtractTransformAtTime(TrackName, Time);
+}
+
 FTransform UContextualAnimSceneAsset::GetAlignmentTransform(int32 SectionIdx, int32 AnimSetIdx, int32 AnimTrackIdx, int32 WarpPointIdx, float Time) const
 {
 	if (const FContextualAnimTrack* AnimTrack = GetAnimTrack(SectionIdx, AnimSetIdx, AnimTrackIdx))
@@ -759,6 +764,49 @@ FTransform UContextualAnimSceneAsset::GetAlignmentTransformForRoleRelativeToOthe
 	return FTransform::Identity;
 }
 
+FTransform UContextualAnimSceneAsset::GetAlignmentTransform(const FContextualAnimSceneBinding& InBinding, const FContextualAnimSceneBindings& InBindings, FName WarpPointName, float Time) const
+{
+	if (InBinding.GetAlignmentData().Tracks.GetNum() > 0)
+	{
+		return InBinding.GetAlignmentData().ExtractTransformAtTime(WarpPointName, Time);
+	}
+
+	const FContextualAnimSceneSection* Section = GetSection(InBinding.GetSectionIdx());
+	if (Section == nullptr)
+	{
+		UE_LOG(LogContextualAnim, Warning, TEXT("UContextualAnimSceneAsset::GetAlignmentTransform Failed. Reason: Can't get AnimSceneSection. Asset: %s Section Idx: %d"), *GetNameSafe(this), InBinding.GetSectionIdx());
+		return FTransform::Identity;
+	}
+
+	const FContextualAnimSet* AnimSet = GetAnimSet(InBinding.GetSectionIdx(), InBinding.GetAnimSetIdx());
+	if (AnimSet == nullptr)
+	{
+		UE_LOG(LogContextualAnim, Warning, TEXT("UContextualAnimSceneAsset::GetAlignmentTransform Failed. Reason: Can't get AnimSet. Asset: %s Section Idx: %d Anim Set Idx: %d"), *GetNameSafe(this), InBinding.GetSectionIdx(), InBinding.GetAnimSetIdx());
+		return FTransform::Identity;
+	}
+
+	const UAnimSequenceBase* Animation = InBinding.GetAnimation();
+
+	if (const FTransform* WarpPointTransformPtr = AnimSet->WarpPoints.Find(WarpPointName))
+	{
+		// TODO: Add possible override the Role Mesh to Component As Well
+		const FName Role = InBindings.GetRoleFromBinding(InBinding);
+		const FTransform MeshToComponentInverse = FTransform(GetMeshToComponentForRole(Role).GetRotation()).Inverse();
+
+		if (Animation)
+		{
+			const FTransform RootTransform = MeshToComponentInverse * (UContextualAnimUtilities::ExtractRootTransformFromAnimation(Animation, Time) * InBinding.GetMeshToScene());
+			return RootTransform.GetRelativeTransform(*WarpPointTransformPtr);
+		}
+		else
+		{
+			const FTransform RootTransform = MeshToComponentInverse * InBinding.GetMeshToScene();
+			return RootTransform.GetRelativeTransform(*WarpPointTransformPtr);
+		}
+	}
+	return  FTransform::Identity;
+}
+
 const FContextualAnimTrack* UContextualAnimSceneAsset::FindAnimTrackForRoleWithClosestEntryLocation(int32 SectionIdx, const FName& Role, const FContextualAnimSceneBindingContext& Primary, const FVector& TestLocation) const
 {
 	const FContextualAnimTrack* Result = nullptr;
@@ -920,6 +968,11 @@ FTransform UContextualAnimSceneAsset::BP_GetIKTargetTransformForRoleAtTime(int32
 	return (Sections.IsValidIndex(SectionIdx)) ? Sections[SectionIdx].GetIKTargetTransformForRoleAtTime(AnimSetIdx, Role, TrackName, Time) : FTransform::Identity;
 }
 
+FTransform UContextualAnimSceneAsset::BP_GetIKTargetTransformForBindingAtTime(const FContextualAnimSceneBinding& InBinding, FName TrackName, float Time) const
+{
+	return InBinding.GetIKTargetData().ExtractTransformAtTime(TrackName, Time);
+}
+
 int32 UContextualAnimSceneAsset::BP_FindAnimSetIndexByAnimation(int32 SectionIdx, const UAnimSequenceBase* Animation) const
 {
 	if (Sections.IsValidIndex(SectionIdx))
diff --git a/Engine/Plugins/Experimental/Animation/ContextualAnimation/Source/ContextualAnimation/Private/ContextualAnimTypes.cpp b/Engine/Plugins/Experimental/Animation/ContextualAnimation/Source/ContextualAnimation/Private/ContextualAnimTypes.cpp
--- a/Engine/Plugins/Experimental/Animation/ContextualAnimation/Source/ContextualAnimation/Private/ContextualAnimTypes.cpp
+++ b/Engine/Plugins/Experimental/Animation/ContextualAnimation/Source/ContextualAnimation/Private/ContextualAnimTypes.cpp
@@ -14,12 +14,23 @@
 #include "GameFramework/CharacterMovementComponent.h"
 #include "MotionWarpingComponent.h"
 #include "Algo/RandomShuffle.h"
+#include "ContextualAnimOverrideInterface.h"
 
 #include UE_INLINE_GENERATED_CPP_BY_NAME(ContextualAnimTypes)
 
 DEFINE_LOG_CATEGORY(LogContextualAnim);
 
 const FContextualAnimSceneBinding FContextualAnimSceneBinding::InvalidBinding;
+
+void FContextualAnimSceneBinding::SetOverrideData(const FContextualAnimationOverrideData& InData)
+{
+	Animation = InData.AnimSequence;
+	MeshToScene = InData.MeshToScene;
+
+	AlignmentData = InData.AlignmentData;
+	IKTargetData = InData.IKTargetData;
+}
+
 const FContextualAnimTrack FContextualAnimTrack::EmptyTrack;
 const FContextualAnimIKTarget FContextualAnimIKTarget::InvalidIKTarget;
 const FContextualAnimIKTargetDefContainer FContextualAnimIKTargetDefContainer::EmptyContainer;
@@ -324,11 +335,25 @@ FContextualAnimSceneBinding::FContextualAnimSceneBinding(const FContextualAnimSc
 	SetAnimTrack(InAnimTrack);
 }
 
-void FContextualAnimSceneBinding::SetAnimTrack(const FContextualAnimTrack& InAnimTrack) 
+FTransform FContextualAnimSceneBinding::GetIKTargetTransform(const FName& TrackName, float Time) const
+{
+	return IKTargetData.ExtractTransformAtTime(TrackName, Time);
+}
+
+FTransform FContextualAnimSceneBinding::GetAlignmentTransform(const FName& TrackName, float Time) const
+{
+	return AlignmentData.ExtractTransformAtTime(TrackName, Time);
+}
+
+void FContextualAnimSceneBinding::SetAnimTrack(const FContextualAnimTrack& InAnimTrack)
 {
 	SectionIdx = InAnimTrack.SectionIdx;
 	AnimSetIdx = InAnimTrack.AnimSetIdx;
 	AnimTrackIdx = InAnimTrack.AnimTrackIdx;
+	Animation = InAnimTrack.Animation;
+	AlignmentData = InAnimTrack.AlignmentData;
+	IKTargetData = InAnimTrack.IKTargetData;
+	MeshToScene = InAnimTrack.MeshToScene;
 }
 
 FAnimMontageInstance* FContextualAnimSceneBinding::GetAnimMontageInstance() const
@@ -510,6 +535,8 @@ bool FContextualAnimSceneBindings::BindActorToRole(AActor& ActorRef, FName Role,
 
 	Data.Add(FContextualAnimSceneBinding(FContextualAnimSceneBindingContext(&ActorRef), *AnimTrackPtr));
 
+	UpdateBindingOverridesIfNeeded();
+
 	return true;
 }
 
@@ -556,6 +583,11 @@ FTransform FContextualAnimSceneBindings::GetIKTargetTransformFromBinding(const F
 		return FTransform::Identity;
 	}
 
+	if (Binding.GetIKTargetData().Tracks.GetNum() > 0)
+	{
+		return Binding.GetIKTargetTransform(TrackName, Time);
+	}
+
 	return GetSceneAsset()->GetIKTargetTransform(Binding.SectionIdx, Binding.AnimSetIdx, Binding.AnimTrackIdx, TrackName, Time);
 }
 
@@ -569,6 +601,11 @@ FTransform FContextualAnimSceneBindings::GetAlignmentTransformFromBinding(const
 		return FTransform::Identity;
 	}
 
+	if (Binding.GetAlignmentData().Tracks.GetNum() > 0)
+	{
+		return Binding.GetAlignmentTransform(TrackName, Time);
+	}
+
 	return GetSceneAsset()->GetAlignmentTransform(Binding.SectionIdx, Binding.AnimSetIdx, Binding.AnimTrackIdx, TrackName, Time);
 }
 
@@ -1013,6 +1050,8 @@ void FContextualAnimSceneBindings::TransitionTo(int32 NewSectionIdx, int32 NewAn
 				Binding.AnimTrackIdx = AnimTrack->AnimTrackIdx;
 			}
 		}
+
+		UpdateBindingOverridesIfNeeded();
 	}
 }
 
@@ -1028,4 +1067,49 @@ bool FContextualAnimSceneBindings::RemoveActor(AActor& ActorRef)
 	}
 
 	return false;
-}
\ No newline at end of file
+}
+
+void FContextualAnimSceneBindings::SetOverrideProvider(const TScriptInterface <IContextualAnimOverrideInterface>& Provider)
+{
+	OverrideProviderInterface = Provider;	
+	UpdateBindingOverridesIfNeeded();
+}
+
+bool FContextualAnimSceneBindings::UpdateBindingOverridesIfNeeded()
+{
+	if (OverrideProviderInterface == nullptr)
+	{
+		return false;
+	}
+
+	if (SceneAsset == nullptr)
+	{
+		return false;
+	}
+
+	for (FContextualAnimSceneBinding& Binding : Data)
+	{
+		const FContextualAnimTrack* AnimTrack = SceneAsset->GetAnimTrack(Binding.SectionIdx, Binding.AnimSetIdx, GetRoleFromBinding(Binding));
+		if (AnimTrack)
+		{
+			FContextualAnimationOverrideData OverrideData(*AnimTrack);
+			if (OverrideProviderInterface->GetCASAnimationDataOverride(Binding, *this, *AnimTrack, OverrideData))
+			{
+				Binding.SetOverrideData(OverrideData);
+			}
+		}
+	}
+
+	for (FContextualAnimSceneBinding& Binding : Data)
+	{
+		// We need to "Recompute" these tracks with the updated data
+		if (SceneAsset->ShouldPrecomputeAlignmentTracks())
+		{
+			UContextualAnimUtilities::GenerateAlignmentTrack(*SceneAsset, Binding, *this, Binding.GetAlignmentData());
+		}
+
+		UContextualAnimUtilities::GenerateIKTargetTrack(*SceneAsset, Binding, *this, Binding.GetIKTargetData());
+	}
+
+	return true;
+}
diff --git a/Engine/Plugins/Experimental/Animation/ContextualAnimation/Source/ContextualAnimation/Private/ContextualAnimUtilities.cpp b/Engine/Plugins/Experimental/Animation/ContextualAnimation/Source/ContextualAnimation/Private/ContextualAnimUtilities.cpp
--- a/Engine/Plugins/Experimental/Animation/ContextualAnimation/Source/ContextualAnimation/Private/ContextualAnimUtilities.cpp
+++ b/Engine/Plugins/Experimental/Animation/ContextualAnimation/Source/ContextualAnimation/Private/ContextualAnimUtilities.cpp
@@ -17,6 +17,7 @@
 #include "AnimNotifyState_MotionWarping.h"
 #include "RootMotionModifier.h"
 #include "ContextualAnimSceneActorComponent.h"
+#include "ContextualAnimOverrideInterface.h"
 
 #include UE_INLINE_GENERATED_CPP_BY_NAME(ContextualAnimUtilities)
 
@@ -264,6 +265,409 @@ FAnimMontageInstance* UContextualAnimUtilities::TryGetActiveAnimMontageInstance(
 	return nullptr;
 }
 
+float UContextualAnimUtilities::GetSyncTimeForWarpSection(const UAnimSequenceBase* Animation, const FName& WarpSectionName)
+{
+	float StartTime, EndTime;
+	GetStartAndEndTimeForWarpSection(Animation, WarpSectionName, StartTime, EndTime);
+	return EndTime;
+}
+
+void UContextualAnimUtilities::GetStartAndEndTimeForWarpSection(const UAnimSequenceBase* Animation, const FName& WarpSectionName, float& OutStartTime, float& OutEndTime)
+{
+	//@TODO: We need a better way to identify warping sections within the animation. This is just a temp solution
+	//@TODO: We should cache this data
+
+	OutStartTime = 0.f;
+	OutEndTime = 0.f;
+
+	int32 Index = INDEX_NONE;
+	float LastEndTime = 0.f;
+	if (Animation && WarpSectionName != NAME_None)
+	{
+		for (int32 Idx = 0; Idx < Animation->Notifies.Num(); Idx++)
+		{
+			const FAnimNotifyEvent& NotifyEvent = Animation->Notifies[Idx];
+			if (const UAnimNotifyState_MotionWarping* Notify = Cast<const UAnimNotifyState_MotionWarping>(NotifyEvent.NotifyStateClass))
+			{
+				if (const URootMotionModifier_Warp* Config = Cast<const URootMotionModifier_Warp>(Notify->RootMotionModifier))
+				{
+					const FName WarpTargetName = Config->WarpTargetName;
+					if (WarpSectionName == WarpTargetName)
+					{
+						const float NotifyEndTriggerTime = NotifyEvent.GetEndTriggerTime();
+						if (NotifyEndTriggerTime > LastEndTime)
+						{
+							LastEndTime = NotifyEndTriggerTime;
+							Index = Idx;
+						}
+					}
+				}
+			}
+		}
+	}
+
+	if (Index != INDEX_NONE)
+	{
+		const FAnimNotifyEvent& NotifyEvent = Animation->Notifies[Index];
+		OutStartTime = NotifyEvent.GetTriggerTime();
+		OutEndTime = NotifyEvent.GetEndTriggerTime();
+	}
+}
+
+void UContextualAnimUtilities::GetStartAndEndTimeForWarpSection(const UAnimSequenceBase* Animation, int32 WarpSectionIndex, float& OutStartTime, float& OutEndTime)
+{
+	OutStartTime = 0.f;
+	OutEndTime = 0.f;
+
+	if (Animation && WarpSectionIndex >= 0)
+	{
+		FName LastWarpTargetName = NAME_None;
+		int32 LastWarpSectionIndex = INDEX_NONE;
+
+		for (int32 Idx = 0; Idx < Animation->Notifies.Num(); Idx++)
+		{
+			const FAnimNotifyEvent& NotifyEvent = Animation->Notifies[Idx];
+			if (const UAnimNotifyState_MotionWarping* Notify = Cast<const UAnimNotifyState_MotionWarping>(NotifyEvent.NotifyStateClass))
+			{
+				if (const URootMotionModifier_Warp* Modifier = Cast<const URootMotionModifier_Warp>(Notify->RootMotionModifier))
+				{
+					const FName WarpTargetName = Modifier->WarpTargetName;
+					if (WarpTargetName != NAME_None)
+					{
+						// First valid warping window. Initialize everything
+						if (LastWarpSectionIndex == INDEX_NONE)
+						{
+							LastWarpTargetName = WarpTargetName;
+							OutStartTime = NotifyEvent.GetTriggerTime();
+							OutEndTime = NotifyEvent.GetEndTriggerTime();
+							LastWarpSectionIndex = 0;
+						}
+						// If we hit another warping window but the sync point is the same as the previous, update SyncTime.
+						// This is to deal with cases where a first short window is used to face the alignment point and a second one to perform the rest of the warp
+						else if (WarpTargetName == LastWarpTargetName)
+						{
+							OutStartTime = NotifyEvent.GetTriggerTime();
+							OutEndTime = NotifyEvent.GetEndTriggerTime();
+						}
+						// If we hit another warping window but with a different sync point name means that we have hit the first window of another warping section
+						else
+						{
+							// If we haven't reached the desired WarpSection yet. Update control vars and keep moving
+							if (WarpSectionIndex > LastWarpSectionIndex)
+							{
+								LastWarpTargetName = WarpTargetName;
+								OutStartTime = NotifyEvent.GetTriggerTime();
+								OutEndTime = NotifyEvent.GetEndTriggerTime();
+								LastWarpSectionIndex++;
+							}
+							// Otherwise, stop here and return the value of the last window we found
+							else
+							{
+								break;
+							}
+						}
+					}
+				}
+			}
+		}
+	}
+}
+
+void UContextualAnimUtilities::GenerateAlignmentTrack(const UContextualAnimSceneAsset& SceneAsset, const FContextualAnimSceneBinding& InBinding, const FContextualAnimSceneBindings& InBindings, FContextualAnimAlignmentTrackContainer& Container)
+{
+	if (!SceneAsset.ShouldPrecomputeAlignmentTracks())
+	{
+		Container.Empty();
+		return;
+	}
+
+	// Necessary for FCompactPose that uses a FAnimStackAllocator (TMemStackAllocator) which allocates from FMemStack.
+	// When allocating memory from FMemStack we need to explicitly use FMemMark to ensure items are freed when the scope exits. 
+	// UWorld::Tick adds a FMemMark to catch any allocation inside the game tick 
+	// but any allocation from outside the game tick (like here when generating the alignment tracks off-line) must explicitly add a mark to avoid a leak 
+	FMemMark Mark(FMemStack::Get());
+
+	//UE_LOG(LogContextualAnim, Log, TEXT("Generating AlignmentTracks Tracks. Animation: %s"), *GetNameSafe(AnimTrack.Animation));
+
+	int32 SectionIndex = InBinding.GetSectionIdx();
+	int32 AnimSetIndex = InBinding.GetAnimSetIdx();
+
+	const FContextualAnimSceneSection* AnimSection = SceneAsset.GetSection(SectionIndex);
+	const FContextualAnimSet* AnimSet = AnimSection->GetAnimSet(AnimSetIndex);
+
+	const FName& Role = InBindings.GetRoleFromBinding(InBinding);
+
+	const FTransform MeshToComponentInverse = FTransform(SceneAsset.GetMeshToComponentForRole(Role).GetRotation()).Inverse();
+	const float SampleInterval = 1.f / SceneAsset.GetSampleRate();
+
+	// Initialize tracks for each alignment section
+	const int32 TotalTracks = AnimSection->GetWarpPointDefinitions().Num();
+	Container.Initialize(TotalTracks, SampleInterval);
+	for (int32 Idx = 0; Idx < TotalTracks; Idx++)
+	{
+		Container.Tracks.TrackNames.Add(AnimSection->GetWarpPointDefinitions()[Idx].WarpTargetName);
+		Container.Tracks.AnimationTracks.AddZeroed();
+	}
+
+	if (const UAnimSequenceBase* Animation = InBinding.GetAnimation())
+	{
+		float Time = 0.f;
+		float EndTime = Animation->GetPlayLength();
+		int32 SampleIndex = 0;
+		while (Time < EndTime)
+		{
+			Time = FMath::Clamp(SampleIndex * SampleInterval, 0.f, EndTime);
+			SampleIndex++;
+
+			const FTransform RootTransform = MeshToComponentInverse * (UContextualAnimUtilities::ExtractRootTransformFromAnimation(Animation, Time) * InBinding.GetMeshToScene());
+
+			for (int32 Idx = 0; Idx < TotalTracks; Idx++)
+			{
+				const FTransform* WarpPointTransform = AnimSet->WarpPoints.Find(AnimSection->GetWarpPointDefinitions()[Idx].WarpTargetName);
+				const FTransform RootRelativeToWarpPoint = WarpPointTransform ? RootTransform.GetRelativeTransform(*WarpPointTransform) : RootTransform;
+
+				FRawAnimSequenceTrack& AlignmentTrack = Container.Tracks.AnimationTracks[Idx];
+				AlignmentTrack.PosKeys.Add(FVector3f(RootRelativeToWarpPoint.GetLocation()));
+				AlignmentTrack.RotKeys.Add(FQuat4f(RootRelativeToWarpPoint.GetRotation()));
+			}
+		}
+	}
+	else
+	{
+		const FTransform RootTransform = MeshToComponentInverse *InBinding.GetMeshToScene();
+
+		for (int32 Idx = 0; Idx < TotalTracks; Idx++)
+		{
+			const FTransform* WarpPointTransform = AnimSet->WarpPoints.Find(AnimSection->GetWarpPointDefinitions()[Idx].WarpTargetName);
+			const FTransform RootRelativeToWarpPoint = WarpPointTransform ? RootTransform.GetRelativeTransform(*WarpPointTransform) : RootTransform;
+
+			FRawAnimSequenceTrack& SceneTrack = Container.Tracks.AnimationTracks[Idx];
+			SceneTrack.PosKeys.Add(FVector3f(RootRelativeToWarpPoint.GetLocation()));
+			SceneTrack.RotKeys.Add(FQuat4f(RootRelativeToWarpPoint.GetRotation()));
+		}
+	}
+}
+
+void UContextualAnimUtilities::ExtractPoseIgnoringForceRootLock(UAnimSequenceBase* AnimSequenceBase, const FBoneContainer& BoneContainer, float Time, bool bExtractRootMotion, FCSPose<FCompactPose>& OutPose)
+{
+	UAnimSequence* AnimSequence = nullptr;
+	if (UAnimMontage* AnimMontage = Cast<UAnimMontage>(AnimSequenceBase))
+	{
+		if (AnimMontage->SlotAnimTracks.Num() > 0)
+		{
+			const float ClampedTime = FMath::Clamp(Time, 0.f, AnimMontage->CalculateSequenceLength());
+			if (FAnimSegment* Segment = AnimMontage->SlotAnimTracks[0].AnimTrack.GetSegmentAtTime(ClampedTime))
+			{
+				AnimSequence = Cast<UAnimSequence>(Segment->GetAnimReference());
+			}
+		}
+	}
+	else
+	{
+		AnimSequence = Cast<UAnimSequence>(AnimSequenceBase);
+	}
+
+	if (AnimSequence)
+	{
+		TGuardValue<bool> ForceRootLockGuard(AnimSequence->bForceRootLock, false);
+		UContextualAnimUtilities::ExtractComponentSpacePose(AnimSequenceBase, BoneContainer, Time, bExtractRootMotion, OutPose);
+	}
+}
+
+FCompactPoseBoneIndex UContextualAnimUtilities::GetCompactPoseBoneIndexFromPose(const FCSPose<FCompactPose>& Pose, const FName& BoneName)
+{
+	const FBoneContainer& BoneContainer = Pose.GetPose().GetBoneContainer();
+	for (int32 Idx = Pose.GetPose().GetNumBones() - 1; Idx >= 0; Idx--)
+	{
+		if (BoneContainer.GetReferenceSkeleton().GetBoneName(BoneContainer.GetBoneIndicesArray()[Idx]) == BoneName)
+		{
+			return FCompactPoseBoneIndex(Idx);
+		}
+	}
+
+	checkf(false, TEXT("BoneName: %s Pose.Asset: %s Pose.NumBones: %d"), *BoneName.ToString(), *GetNameSafe(Pose.GetPose().GetBoneContainer().GetAsset()), Pose.GetPose().GetNumBones());
+	return FCompactPoseBoneIndex(INDEX_NONE);
+}
+
+void UContextualAnimUtilities::GenerateIKTargetTrack(const UContextualAnimSceneAsset& SceneAsset, const FContextualAnimSceneBinding& InBinding, const FContextualAnimSceneBindings& InBindings, FContextualAnimAlignmentTrackContainer& Container)
+{
+	// @TODO: Optimize
+	// We are generating IK Target tracks for the duration of the animation and for all the sets in the section
+	// We should generate those tracks only for the animations with IK windows and only for the range of that window
+
+	FMemMark Mark(FMemStack::Get());
+
+	const FName& Role = InBindings.GetRoleFromBinding(InBinding);
+	const TArray<FContextualAnimIKTargetDefinition>& IKTargetDefs = SceneAsset.GetIKTargetDefsForRole(Role).IKTargetDefs;
+	
+	if (IKTargetDefs.Num() == 0)
+	{
+		return;
+	}
+
+	if (UAnimSequenceBase* Animation = InBinding.GetAnimation())
+	{
+		UE_LOG(LogContextualAnim, Log, TEXT("Generating IK Target Tracks. Animation: %s"), *GetNameSafe(Animation));
+
+		const float SampleInterval = 1.f / SceneAsset.GetSampleRate();
+
+		TArray<FBoneIndexType> RequiredBoneIndexArray;
+
+		// Helper structure to group pose extraction per target so we can extract the pose for all the bones that are relative to the same target in one pass
+		struct FPoseExtractionHelper
+		{
+			//const FContextualAnimTrack* TargetAnimTrackPtr = nullptr;
+			const FContextualAnimSceneBinding* TargetBindingPtr = nullptr;
+			TArray<TTuple<FName, FName, int32, FName, int32>> BonesData; //0: GoalName, 1: MyBoneName 2: MyBoneIndex, 3: TargetBoneName, 4: TargetBoneIndex
+		};
+
+		TMap<FName, FPoseExtractionHelper> PoseExtractionHelperMap;
+		PoseExtractionHelperMap.Reserve(IKTargetDefs.Num());
+
+
+		int32 TotalTracks = 0;
+		for (const FContextualAnimIKTargetDefinition& IKTargetDef : IKTargetDefs)
+		{
+			if (IKTargetDef.Provider == EContextualAnimIKTargetProvider::Autogenerated)
+			{
+				const FName TargetRole = IKTargetDef.TargetRoleName;
+
+				FPoseExtractionHelper* DataPtr = PoseExtractionHelperMap.Find(TargetRole);
+				if (DataPtr == nullptr)
+				{
+					// Find track for target role. 
+					const FContextualAnimSceneBinding* TargetBinding = InBindings.FindBindingByRole(TargetRole);
+					
+					DataPtr = &PoseExtractionHelperMap.Add(TargetRole);
+					DataPtr->TargetBindingPtr = TargetBinding;
+				}
+
+				const FName BoneName = IKTargetDef.BoneName;
+				const int32 BoneIndex = Animation->GetSkeleton()->GetReferenceSkeleton().FindBoneIndex(BoneName);
+				if (BoneIndex == INDEX_NONE)
+				{
+					UE_LOG(LogContextualAnim, Warning, TEXT("\t Can't find BoneIndex. BoneName: %s Animation: %s Skel: %s"),
+						*BoneName.ToString(), *GetNameSafe(Animation), *GetNameSafe(Animation->GetSkeleton()));
+
+					continue;
+				}
+
+				// Find TargetBoneIndex. Note that we add TargetBoneIndex even if it is INDEX_NONE. In this case, my bone will be relative to the origin of the target actor. 
+				// This is to support cases where the target actor doesn't have animation or TargetBoneName is None
+				FName TargetBoneName = IKTargetDef.TargetBoneName;
+				const UAnimSequenceBase* TargetAnimation = DataPtr->TargetBindingPtr->GetAnimation();
+
+				const int32 TargetBoneIndex = TargetAnimation ? TargetAnimation->GetSkeleton()->GetReferenceSkeleton().FindBoneIndex(TargetBoneName) : INDEX_NONE;
+				if (TargetBoneIndex == INDEX_NONE)
+				{
+					UE_LOG(LogContextualAnim, Log, TEXT("\t Can't find TargetBoneIndex. BoneName: %s Animation: %s Skel: %s. Track for this bone will be relative to the origin of the target role."),
+						*TargetBoneName.ToString(), *GetNameSafe(TargetAnimation), TargetAnimation ? *GetNameSafe(TargetAnimation->GetSkeleton()) : nullptr);
+
+					TargetBoneName = NAME_None;
+				}
+
+				RequiredBoneIndexArray.AddUnique(BoneIndex);
+
+				DataPtr->BonesData.Add(MakeTuple(IKTargetDef.GoalName, BoneName, BoneIndex, TargetBoneName, TargetBoneIndex));
+				TotalTracks++;
+
+				UE_LOG(LogContextualAnim, Log, TEXT("\t Bone added for extraction. GoalName: %s BoneName: %s (%d) TargetRole: %s TargetAnimation: %s TargetBone: %s (%d)"),
+					*IKTargetDef.GoalName.ToString(), *BoneName.ToString(), BoneIndex, *TargetRole.ToString(), *GetNameSafe(TargetAnimation), *TargetBoneName.ToString(), TargetBoneIndex);
+			}
+		}
+
+		if (TotalTracks > 0)
+		{
+			// Complete bones chain and create bone container to extract pose from my animation
+			Animation->GetSkeleton()->GetReferenceSkeleton().EnsureParentsExistAndSort(RequiredBoneIndexArray);
+			FBoneContainer BoneContainer = FBoneContainer(RequiredBoneIndexArray, UE::Anim::FCurveFilterSettings(UE::Anim::ECurveFilterMode::DisallowAll), *Animation->GetSkeleton());
+
+			// Initialize track container
+			Container.Initialize(TotalTracks, SampleInterval);
+			Container.Tracks.TrackNames.AddZeroed(TotalTracks);
+			Container.Tracks.AnimationTracks.AddZeroed(TotalTracks);
+
+			float Time = 0.f;
+			float EndTime = Animation->GetPlayLength();
+			int32 SampleIndex = 0;
+			while (Time < EndTime)
+			{
+				Time = FMath::Clamp(SampleIndex * SampleInterval, 0.f, EndTime);
+				SampleIndex++;
+
+				// Extract pose from my animation
+				FCSPose<FCompactPose> ComponentSpacePose;
+				UContextualAnimUtilities::ExtractPoseIgnoringForceRootLock(Animation, BoneContainer, Time, false, ComponentSpacePose);
+
+				// For each target role
+				for (auto& Data : PoseExtractionHelperMap)
+				{
+					// Extract pose from target animation if any
+					FCSPose<FCompactPose> OtherComponentSpacePose;
+					TArray<FBoneIndexType> OtherRequiredBoneIndexArray;
+					FBoneContainer OtherBoneContainer;
+					UAnimSequenceBase* OtherAnimation = Data.Value.TargetBindingPtr->GetAnimation();
+					if (OtherAnimation)
+					{
+						// Prepare array with the indices of the bones to extract from target animation
+						OtherRequiredBoneIndexArray.Reserve(Data.Value.BonesData.Num());
+						for (int32 Idx = 0; Idx < Data.Value.BonesData.Num(); Idx++)
+						{
+							const int32 TargetBoneIndex = Data.Value.BonesData[Idx].Get<4>();
+							if (TargetBoneIndex != INDEX_NONE)
+							{
+								OtherRequiredBoneIndexArray.AddUnique(TargetBoneIndex);
+							}
+						}
+
+						if (OtherRequiredBoneIndexArray.Num() > 0)
+						{
+							// Complete bones chain and create bone container to extract pose form the target animation
+							OtherAnimation->GetSkeleton()->GetReferenceSkeleton().EnsureParentsExistAndSort(OtherRequiredBoneIndexArray);
+							OtherBoneContainer = FBoneContainer(OtherRequiredBoneIndexArray, UE::Anim::FCurveFilterSettings(UE::Anim::ECurveFilterMode::DisallowAll), *OtherAnimation->GetSkeleton());
+
+							// Extract pose from target animation
+							UContextualAnimUtilities::ExtractPoseIgnoringForceRootLock(OtherAnimation, OtherBoneContainer, Time, false, OtherComponentSpacePose);
+						}
+					}
+
+					for (int32 Idx = 0; Idx < Data.Value.BonesData.Num(); Idx++)
+					{
+						const FName TrackName = Data.Value.BonesData[Idx].Get<0>();
+						Container.Tracks.TrackNames[Idx] = TrackName;
+
+						// Get bone transform from my animation
+						const FName BoneName = Data.Value.BonesData[Idx].Get<1>();
+						const FCompactPoseBoneIndex BoneIndex = UContextualAnimUtilities::GetCompactPoseBoneIndexFromPose(ComponentSpacePose, BoneName);
+						const FTransform BoneTransform = (ComponentSpacePose.GetComponentSpaceTransform(BoneIndex) * InBinding.GetMeshToScene());
+
+						// Get bone transform from target animation
+						FTransform OtherBoneTransform = Data.Value.TargetBindingPtr->GetMeshToScene();
+						const FName TargetBoneName = Data.Value.BonesData[Idx].Get<3>();
+						if (TargetBoneName != NAME_None)
+						{
+							const FCompactPoseBoneIndex OtherBoneIndex = UContextualAnimUtilities::GetCompactPoseBoneIndexFromPose(OtherComponentSpacePose, TargetBoneName);
+							OtherBoneTransform = (OtherComponentSpacePose.GetComponentSpaceTransform(OtherBoneIndex) * Data.Value.TargetBindingPtr->GetMeshToScene());
+						}
+
+						// Get transform relative to target
+						const FTransform BoneRelativeToOther = BoneTransform.GetRelativeTransform(OtherBoneTransform);
+
+						// Add transform to the track
+						FRawAnimSequenceTrack& NewTrack = Container.Tracks.AnimationTracks[Idx];
+						NewTrack.PosKeys.Add(FVector3f(BoneRelativeToOther.GetLocation()));
+						NewTrack.RotKeys.Add(FQuat4f(BoneRelativeToOther.GetRotation()));
+
+						UE_LOG(LogContextualAnim, Verbose, TEXT("\t\t Animation: %s Time: %f BoneName: %s (T: %s) Target Animation: %s TargetBoneName: %s (T: %s)"),
+							*GetNameSafe(Animation), Time, *BoneName.ToString(), *BoneTransform.GetLocation().ToString(),
+							*GetNameSafe(OtherAnimation), *TargetBoneName.ToString(), *OtherBoneTransform.GetLocation().ToString());
+					}
+				}
+			}
+		}
+	}
+}
+
 void UContextualAnimUtilities::BP_Montage_GetSectionStartAndEndTime(const UAnimMontage* Montage, int32 SectionIndex, float& OutStartTime, float& OutEndTime)
 {
 	if(Montage)
@@ -379,6 +783,12 @@ bool UContextualAnimUtilities::BP_CreateContextualAnimSceneBindingsForTwoActors(
 	return FContextualAnimSceneBindings::TryCreateBindings(*SceneAsset, SectionIdx, Primary, Secondary, OutBindings);
 }
 
+void UContextualAnimUtilities::BP_AddContextualAnimationOverrideProvider(FContextualAnimSceneBindings& InBindings, TScriptInterface<IContextualAnimOverrideInterface>& OverrideProvider, FContextualAnimSceneBindings& OutBindings)
+{
+	InBindings.SetOverrideProvider(OverrideProvider);
+	OutBindings = InBindings;
+}
+
 // SceneBindings Blueprint Interface
 //------------------------------------------------------------------------------------------
 
@@ -528,8 +938,7 @@ void UContextualAnimUtilities::BP_SceneBindings_AddOrUpdateWarpTargetsForSection
 				{
 					if (UMotionWarpingComponent* MotionWarpComp = Binding.GetActor()->FindComponentByClass<UMotionWarpingComponent>())
 					{
-						const FContextualAnimTrack& AnimTrack = Bindings.GetAnimTrackFromBinding(Binding);
-						const float Time = AnimTrack.GetSyncTimeForWarpSection(WarpPointDef.WarpTargetName);
+						const float Time = UContextualAnimUtilities::GetSyncTimeForWarpSection(Binding.GetAnimation(), WarpPointDef.WarpTargetName);
 						const FTransform TransformRelativeToWarpPoint = Bindings.GetAlignmentTransformFromBinding(Binding, WarpPointDef.WarpTargetName, Time);
 						const FTransform WarpTargetTransform = TransformRelativeToWarpPoint * WarpPoint.Transform;
 						MotionWarpComp->AddOrUpdateWarpTargetFromTransform(WarpPoint.Name, WarpTargetTransform);
@@ -602,7 +1011,7 @@ const UAnimSequenceBase* UContextualAnimUtilities::BP_SceneBinding_GetAnimationF
 		return nullptr;
 	}
 
-	return Bindings.GetAnimTrackFromBinding(Binding).Animation;
+	return Binding.GetAnimation();
 }
 
 FName UContextualAnimUtilities::BP_SceneBinding_GetRoleFromBinding(const FContextualAnimSceneBindings& Bindings, const FContextualAnimSceneBinding& Binding)
@@ -724,4 +1133,81 @@ void UContextualAnimUtilities::BP_SceneBinding_GetSectionAndAnimSetIndices(const
 {
 	SectionIdx = Binding.GetSectionIdx();
 	AnimSetIdx = Binding.GetAnimSetIdx();
-}
\ No newline at end of file
+}
+
+void UContextualAnimUtilities::GenerateAlignmentTracksForBindings(UContextualAnimSceneAsset& SceneAsset, FContextualAnimSceneBindings& InBindings)
+{
+	// Necessary for FCompactPose that uses a FAnimStackAllocator (TMemStackAllocator) which allocates from FMemStack.
+	// When allocating memory from FMemStack we need to explicitly use FMemMark to ensure items are freed when the scope exits. 
+	// UWorld::Tick adds a FMemMark to catch any allocation inside the game tick 
+	// but any allocation from outside the game tick (like here when generating the alignment tracks off-line) must explicitly add a mark to avoid a leak 
+	FMemMark Mark(FMemStack::Get());
+
+	for (FContextualAnimSceneBinding& Binding : InBindings)
+	{
+		const FName Role = InBindings.GetRoleFromBinding(Binding);
+
+		GenerateAlignmentTrackForBinding(SceneAsset, Role, Binding);
+	}
+}
+
+void UContextualAnimUtilities::GenerateAlignmentTrackForBinding(UContextualAnimSceneAsset& SceneAsset, const FName& Role, FContextualAnimSceneBinding& Binding)
+{
+	const FContextualAnimSceneSection& Section = *SceneAsset.GetSection(Binding.GetSectionIdx());
+	const FContextualAnimSet& AnimSet = *Section.GetAnimSet(Binding.GetAnimSetIdx());
+
+	const FTransform MeshToComponentInverse = FTransform(SceneAsset.GetMeshToComponentForRole(Role).GetRotation()).Inverse();
+	const float SampleInterval = 1.f / SceneAsset.GetSampleRate();
+
+	// Initialize tracks for each alignment section
+	const int32 TotalTracks = Section.GetWarpPointDefinitions().Num();
+
+	FContextualAnimAlignmentTrackContainer& AlignmentData = Binding.GetAlignmentData();
+	AlignmentData.Initialize(TotalTracks, SampleInterval);
+
+	for (int32 Idx = 0; Idx < TotalTracks; Idx++)
+	{
+		AlignmentData.Tracks.TrackNames.Add(Section.GetWarpPointDefinitions()[Idx].WarpTargetName);
+		AlignmentData.Tracks.AnimationTracks.AddZeroed();
+	}
+
+	if (const UAnimSequenceBase* Animation = Binding.GetAnimation())
+	{
+		float Time = 0.f;
+		float EndTime = Animation->GetPlayLength();
+		int32 SampleIndex = 0;
+		while (Time < EndTime)
+		{
+			Time = FMath::Clamp(SampleIndex * SampleInterval, 0.f, EndTime);
+			SampleIndex++;
+
+			const FTransform RootTransform = MeshToComponentInverse * (UContextualAnimUtilities::ExtractRootTransformFromAnimation(Animation, Time) * Binding.GetMeshToScene());
+
+			for (int32 Idx = 0; Idx < TotalTracks; Idx++)
+			{
+				FRawAnimSequenceTrack& AlignmentTrack = AlignmentData.Tracks.AnimationTracks[Idx];
+
+				const FTransform* WarpPointTransform = AnimSet.WarpPoints.Find(Section.GetWarpPointDefinitions()[Idx].WarpTargetName);
+				const FTransform RootRelativeToWarpPoint = WarpPointTransform ? RootTransform.GetRelativeTransform(*WarpPointTransform) : RootTransform;
+
+				AlignmentTrack.PosKeys.Add(FVector3f(RootRelativeToWarpPoint.GetLocation()));
+				AlignmentTrack.RotKeys.Add(FQuat4f(RootRelativeToWarpPoint.GetRotation()));
+			}
+		}
+	}
+	else
+	{
+		const FTransform RootTransform = MeshToComponentInverse * Binding.GetMeshToScene();
+
+		for (int32 Idx = 0; Idx < TotalTracks; Idx++)
+		{
+			FRawAnimSequenceTrack& SceneTrack = AlignmentData.Tracks.AnimationTracks[Idx];
+
+			const FTransform* WarpPointTransform = AnimSet.WarpPoints.Find(Section.GetWarpPointDefinitions()[Idx].WarpTargetName);
+			const FTransform RootRelativeToWarpPoint = WarpPointTransform ? RootTransform.GetRelativeTransform(*WarpPointTransform) : RootTransform;
+
+			SceneTrack.PosKeys.Add(FVector3f(RootRelativeToWarpPoint.GetLocation()));
+			SceneTrack.RotKeys.Add(FQuat4f(RootRelativeToWarpPoint.GetRotation()));
+		}
+	}
+}
diff --git a/Engine/Plugins/Experimental/Animation/ContextualAnimation/Source/ContextualAnimation/Public/ContextualAnimOverrideInterface.h b/Engine/Plugins/Experimental/Animation/ContextualAnimation/Source/ContextualAnimation/Public/ContextualAnimOverrideInterface.h
new file mode 100644
--- /dev/null
+++ b/Engine/Plugins/Experimental/Animation/ContextualAnimation/Source/ContextualAnimation/Public/ContextualAnimOverrideInterface.h
@@ -0,0 +1,59 @@
+// Copyright Epic Games, Inc. All Rights Reserved.
+
+#pragma once
+
+#include "UObject/Interface.h"
+
+#include "ContextualAnimTypes.h"
+
+#include "ContextualAnimOverrideInterface.generated.h"
+
+class UAnimSequenceBase;
+
+
+USTRUCT(BlueprintType)
+struct FContextualAnimationOverrideData
+{
+	GENERATED_BODY()
+
+	FContextualAnimationOverrideData()
+	{
+
+	}
+
+	FContextualAnimationOverrideData(const FContextualAnimTrack& InAnimTrack)
+	{
+		AnimSequence = InAnimTrack.Animation;
+		MeshToScene = InAnimTrack.MeshToScene;
+	}
+	
+	UPROPERTY(EditAnywhere, Category = Overrides)
+	TObjectPtr<UAnimSequenceBase> AnimSequence = nullptr;
+
+	UPROPERTY(EditAnywhere, Category = Overrides)
+	FTransform MeshToScene = FTransform::Identity;
+
+	/** Container for alignment tracks */
+	UPROPERTY()
+	FContextualAnimAlignmentTrackContainer AlignmentData;
+
+	/** Container for auto generate IK Target Tracks */
+	UPROPERTY()
+	FContextualAnimAlignmentTrackContainer IKTargetData;
+};
+
+
+UINTERFACE()
+class UContextualAnimOverrideInterface : public UInterface
+{
+	GENERATED_BODY()
+};
+
+class IContextualAnimOverrideInterface : public IInterface
+{
+	GENERATED_BODY()
+
+public:
+	virtual bool GetCASAnimationDataOverride(const FContextualAnimSceneBinding& InBinding, const FContextualAnimSceneBindings& InBindings, const FContextualAnimTrack& InAnimTrack, FContextualAnimationOverrideData& OutData) = 0;
+
+};
diff --git a/Engine/Plugins/Experimental/Animation/ContextualAnimation/Source/ContextualAnimation/Public/ContextualAnimSceneAsset.h b/Engine/Plugins/Experimental/Animation/ContextualAnimation/Source/ContextualAnimation/Public/ContextualAnimSceneAsset.h
--- a/Engine/Plugins/Experimental/Animation/ContextualAnimation/Source/ContextualAnimation/Public/ContextualAnimSceneAsset.h
+++ b/Engine/Plugins/Experimental/Animation/ContextualAnimation/Source/ContextualAnimation/Public/ContextualAnimSceneAsset.h
@@ -285,11 +285,13 @@ public:
  	UE_API const FContextualAnimTrack* GetAnimTrack(int32 SectionIdx, int32 AnimSetIdx, int32 AnimTrackIdx) const;
 
 	UE_API FTransform GetIKTargetTransform(int32 SectionIdx, int32 AnimSetIdx, int32 AnimTrackIdx, const FName& TrackName, float Time) const;
+	UE_API FTransform GetIKTargetTransform(const FContextualAnimSceneBinding& InBinding, const FName& TrackName, float Time) const;
 
 	UE_API FTransform GetAlignmentTransform(int32 SectionIdx, int32 AnimSetIdx, int32 AnimTrackIdx, int32 WarpPointIdx, float Time) const;
 	UE_API FTransform GetAlignmentTransform(int32 SectionIdx, int32 AnimSetIdx, int32 AnimTrackIdx, const FName& WarpPointName, float Time) const;
 	UE_API FTransform GetAlignmentTransform(const FContextualAnimTrack& AnimTrack, int32 WarpPointIdx, float Time) const;
 	UE_API FTransform GetAlignmentTransform(const FContextualAnimTrack& AnimTrack, const FName& WarpPointName, float Time) const;
+	UE_API FTransform GetAlignmentTransform(const FContextualAnimSceneBinding& InBinding, const FContextualAnimSceneBindings& InBindings, FName WarpPointName, float Time) const;
 	UE_API FTransform GetAlignmentTransformForRoleRelativeToOtherRole(int32 SectionIdx, int32 AnimSetIdx, FName Role, FName OtherRole, float Time) const;
 
 	UE_API const FContextualAnimTrack* FindAnimTrackForRoleWithClosestEntryLocation(int32 SectionIdx, const FName& Role, const FContextualAnimSceneBindingContext& Primary, const FVector& TestLocation) const;
@@ -320,6 +322,9 @@ public:
 
 	UFUNCTION(BlueprintCallable, Category = "Contextual Anim|Scene Asset", meta = (DisplayName = "Get IK Target Transform For Role At Time"))
 	UE_API FTransform BP_GetIKTargetTransformForRoleAtTime(int32 SectionIdx, int32 AnimSetIdx, FName Role, FName TrackName, float Time) const;
+	
+	UFUNCTION(BlueprintCallable, Category = "Contextual Anim|Scene Asset", meta = (DisplayName = "Get IK Target Transform For Binding At Time"))
+	UE_API FTransform BP_GetIKTargetTransformForBindingAtTime(const FContextualAnimSceneBinding& InBinding, FName TrackName, float Time) const;
 
 	UFUNCTION(BlueprintCallable, Category = "Contextual Anim|Scene Asset", meta = (DisplayName = "Get Start and End Time For Warp Section"))
 	UE_API void BP_GetStartAndEndTimeForWarpSection(int32 SectionIdx, int32 AnimSetIdx, FName Role, FName WarpSectionName, float& OutStartTime, float& OutEndTime) const;
@@ -342,6 +347,9 @@ protected:
 	UPROPERTY(EditAnywhere, Category = "Settings", meta = (TitleProperty = "Role"))
 	TArray<FContextualAnimActorPreviewData> OverridePreviewData;
 
+	UPROPERTY(EditAnywhere, Category = "Settings", meta = (TitleProperty = "Role", MustImplement = "/Script/ContextualAnimation.ContextualAnimOverrideInterface"))
+	TSubclassOf<UObject> PreviewOverrideProvider = nullptr;;
+
 #endif // WITH_EDITORONLY_DATA
 
 	UPROPERTY(EditAnywhere, Category = "Defaults")
diff --git a/Engine/Plugins/Experimental/Animation/ContextualAnimation/Source/ContextualAnimation/Public/ContextualAnimTypes.h b/Engine/Plugins/Experimental/Animation/ContextualAnimation/Source/ContextualAnimation/Public/ContextualAnimTypes.h
--- a/Engine/Plugins/Experimental/Animation/ContextualAnimation/Source/ContextualAnimation/Public/ContextualAnimTypes.h
+++ b/Engine/Plugins/Experimental/Animation/ContextualAnimation/Source/ContextualAnimation/Public/ContextualAnimTypes.h
@@ -21,7 +21,10 @@ class UContextualAnimSceneAsset;
 class UContextualAnimSceneActorComponent;
 class UCharacterMovementComponent;
 class UMotionWarpingComponent;
+class IContextualAnimOverrideInterface;
+
 struct FAnimMontageInstance;
+struct FContextualAnimationOverrideData;
 
 namespace UE 
 {
@@ -473,6 +476,9 @@ struct FContextualAnimSceneBinding
 	inline int8 GetSectionIdx() const { return SectionIdx; }
 	inline int8 GetAnimSetIdx() const { return AnimSetIdx; }
 	inline int8 GetAnimTrackIdx() const { return AnimTrackIdx; }
+
+	UE_API FTransform GetIKTargetTransform(const FName& TrackName, float Time) const;
+	UE_API FTransform GetAlignmentTransform(const FName& TrackName, float Time) const;
 	
 	UE_API void SetAnimTrack(const FContextualAnimTrack& InAnimTrack);
 
@@ -488,6 +494,19 @@ struct FContextualAnimSceneBinding
 
 	static UE_API const FContextualAnimSceneBinding InvalidBinding;
 
+	UAnimSequenceBase* GetAnimation() const { return Animation; }
+
+	FContextualAnimAlignmentTrackContainer& GetAlignmentData() { return AlignmentData; }
+	const FContextualAnimAlignmentTrackContainer& GetAlignmentData() const { return AlignmentData; }
+
+	FContextualAnimAlignmentTrackContainer& GetIKTargetData() { return IKTargetData; }
+	const FContextualAnimAlignmentTrackContainer& GetIKTargetData() const { return IKTargetData; }
+
+	const FTransform& GetMeshToScene() const { return MeshToScene; }
+	FTransform& GetMeshToScene() { return MeshToScene; }
+
+	void SetOverrideData(const FContextualAnimationOverrideData& InData);
+
 private:
 
 	friend class UContextualAnimSceneInstance;
@@ -504,6 +523,26 @@ private:
 
 	UPROPERTY()
 	int8 AnimTrackIdx = INDEX_NONE;
+
+	UPROPERTY()
+	TObjectPtr<UAnimSequenceBase> Animation = nullptr;
+
+	/** Container for alignment tracks */
+	UPROPERTY(NotReplicated)
+	FContextualAnimAlignmentTrackContainer AlignmentData;
+
+	/** Container for auto generate IK Target Tracks */
+	UPROPERTY(NotReplicated)
+	FContextualAnimAlignmentTrackContainer IKTargetData;
+	
+	UPROPERTY()
+	FTransform MeshToScene;
+
+	/** Map of WarpTargetNames and Transforms for this set. Generated off line based on warp points defined in the asset */
+	UPROPERTY()
+	TMap<FName, FTransform> WarpPoints;
+
+	friend FContextualAnimSceneBinding;
 };
 
 USTRUCT(BlueprintType)
@@ -594,6 +633,16 @@ struct FContextualAnimSceneBindings
 
 	UE_API bool AreAllBindingsInTheSameSection() const;
 
+	UE_API void SetOverrideProvider(const TScriptInterface <IContextualAnimOverrideInterface>& Provider);
+	
+	UE_API const TScriptInterface <IContextualAnimOverrideInterface>& GetOverrideProvider() const { return OverrideProviderInterface;}
+
+	/*
+	* Updates Each Bindings Overrides, Including Animation, Alignment IK, ETC 
+	* Returns True is there is a valid Scene and Override Provider
+	*/
+	UE_API bool UpdateBindingOverridesIfNeeded();
+
 private:
 
 	friend class UContextualAnimManager;
@@ -615,6 +664,9 @@ private:
 	/** List of actors bound to each role in the SceneAsset */
 	UPROPERTY()
 	TArray<FContextualAnimSceneBinding> Data;
+	
+	UPROPERTY()
+	TScriptInterface<IContextualAnimOverrideInterface> OverrideProviderInterface = nullptr;
 };
 
 // DEPRECATED: Kept around only to do not break existing content. It will go away in the future.
diff --git a/Engine/Plugins/Experimental/Animation/ContextualAnimation/Source/ContextualAnimation/Public/ContextualAnimUtilities.h b/Engine/Plugins/Experimental/Animation/ContextualAnimation/Source/ContextualAnimation/Public/ContextualAnimUtilities.h
--- a/Engine/Plugins/Experimental/Animation/ContextualAnimation/Source/ContextualAnimation/Public/ContextualAnimUtilities.h
+++ b/Engine/Plugins/Experimental/Animation/ContextualAnimation/Source/ContextualAnimation/Public/ContextualAnimUtilities.h
@@ -3,11 +3,14 @@
 #pragma once
 
 #include "ContextualAnimTypes.h"
+#include "ContextualAnimSceneAsset.h"
+
 #include "ContextualAnimUtilities.generated.h"
 
 #define UE_API CONTEXTUALANIMATION_API
 
 class UContextualAnimSceneAsset;
+class IContextualAnimOverrideInterface;
 class UMeshComponent;
 class USkeletalMeshComponent;
 class UStaticMeshComponent;
@@ -67,6 +70,15 @@ public:
 
 	static UE_API FAnimMontageInstance* TryGetActiveAnimMontageInstance(const AActor* Actor);
 
+	static UE_API float GetSyncTimeForWarpSection(const UAnimSequenceBase* Animation, const FName& WarpSectionName);
+	
+	static UE_API void GetStartAndEndTimeForWarpSection(const UAnimSequenceBase* Animation, const FName& WarpSectionName, float& OutStartTime, float& OutEndTime);
+	static UE_API void GetStartAndEndTimeForWarpSection(const UAnimSequenceBase* Animation, int32 WarpSectionIndex, float& OutStartTime, float& OutEndTime);
+	
+	static UE_API void GenerateAlignmentTrack(const UContextualAnimSceneAsset& SceneAsset, const FContextualAnimSceneBinding& InBinding, const FContextualAnimSceneBindings& InBindings, FContextualAnimAlignmentTrackContainer& Container);
+	static UE_API void GenerateIKTargetTrack(const UContextualAnimSceneAsset& SceneAsset, const FContextualAnimSceneBinding& InBinding, const FContextualAnimSceneBindings& InBindings, FContextualAnimAlignmentTrackContainer& Container);
+
+	
 	static UE_API void DrawSector(FPrimitiveDrawInterface& PDI, const FVector& Origin, const FVector& Direction, float MinDistance, float MaxDistance, float MinAngle, float MaxAngle, const FLinearColor& Color, uint8 DepthPriority, float Thickness, bool bDashedLine);
 
 	UFUNCTION(BlueprintCallable, Category = "Contextual Anim|Scene Asset", meta = (DisplayName = "Create Contextual Anim Scene Bindings"))
@@ -75,6 +87,9 @@ public:
 	UFUNCTION(BlueprintCallable, Category = "Contextual Anim|Scene Asset", meta = (DisplayName = "Create Contextual Anim Scene Bindings For Two Actors"))
 	static UE_API bool BP_CreateContextualAnimSceneBindingsForTwoActors(const UContextualAnimSceneAsset* SceneAsset, const FContextualAnimSceneBindingContext& Primary, const FContextualAnimSceneBindingContext& Secondary, FContextualAnimSceneBindings& OutBindings);
 
+	UFUNCTION(BlueprintCallable, Category = "Contextual Anim|Scene Asset", meta = (DisplayName = "Add Override Provider to Bindings"))
+	static UE_API void BP_AddContextualAnimationOverrideProvider(UPARAM(Ref)FContextualAnimSceneBindings& InBindings, UPARAM(Ref)TScriptInterface<IContextualAnimOverrideInterface>& OverrideProvider, FContextualAnimSceneBindings& OutBindings);
+	
 	// Montage Blueprint Interface
 	//------------------------------------------------------------------------------------------
 
@@ -190,6 +205,12 @@ public:
 
 	UFUNCTION(BlueprintPure, Category = "Contextual Anim|Scene Binding", meta = (DisplayName = "Get Section And AnimSet Indices"))
 	static UE_API void BP_SceneBinding_GetSectionAndAnimSetIndices(const FContextualAnimSceneBinding& Binding, int32& SectionIdx, int32& AnimSetIdx);
+
+	static void ExtractPoseIgnoringForceRootLock(UAnimSequenceBase* AnimSequenceBase, const FBoneContainer& BoneContainer, float Time, bool bExtractRootMotion, FCSPose<FCompactPose>& OutPose);
+	static FCompactPoseBoneIndex GetCompactPoseBoneIndexFromPose(const FCSPose<FCompactPose>& Pose, const FName& BoneName);
+	
+	static UE_API void GenerateAlignmentTracksForBindings(UContextualAnimSceneAsset& SceneAsset, FContextualAnimSceneBindings& InBindings);
+	static UE_API void GenerateAlignmentTrackForBinding(UContextualAnimSceneAsset& SceneAsset, const FName& Role, FContextualAnimSceneBinding& Binding);
 };
 
 #undef UE_API
diff --git a/Engine/Plugins/Experimental/Animation/ContextualAnimation/Source/ContextualAnimationEditor/Private/ContextualAnimEdMode.cpp b/Engine/Plugins/Experimental/Animation/ContextualAnimation/Source/ContextualAnimationEditor/Private/ContextualAnimEdMode.cpp
--- a/Engine/Plugins/Experimental/Animation/ContextualAnimation/Source/ContextualAnimationEditor/Private/ContextualAnimEdMode.cpp
+++ b/Engine/Plugins/Experimental/Animation/ContextualAnimation/Source/ContextualAnimationEditor/Private/ContextualAnimEdMode.cpp
@@ -77,10 +77,10 @@ void FContextualAnimEdMode::Render(const FSceneView* View, FViewport* Viewport,
 							// Draw warp target transforms for each actor using this warp point in White. This is where the root should end up at the end of the warping window.
 							for (const FContextualAnimSceneBinding& Binding : Bindings)
 							{
-								const FContextualAnimTrack& AnimTrack = Bindings.GetAnimTrackFromBinding(Binding);
-								if(AnimTrack.Animation)
+								const UAnimSequenceBase* AnimationSequence = Binding.GetAnimation();
+								if(AnimationSequence)
 								{
-									const float Time = AnimTrack.GetSyncTimeForWarpSection(WarpPointDef.WarpTargetName);
+									const float Time = UContextualAnimUtilities::GetSyncTimeForWarpSection(AnimationSequence, WarpPointDef.WarpTargetName);
 									const FTransform TransformRelativeToWarpPoint = SceneAsset->GetAlignmentTransform(Binding.GetSectionIdx(), Binding.GetAnimSetIdx(), Binding.GetAnimTrackIdx(), WarpPointDef.WarpTargetName, Time);
 									const FTransform WarpTargetTransform = (TransformRelativeToWarpPoint * WarpPoint.Transform);
 
@@ -195,7 +195,8 @@ void FContextualAnimEdMode::Render(const FSceneView* View, FViewport* Viewport,
 							const EMultiOptionDrawMode EntryPosesDrawMode = ViewportClient->GetEntryPosesDrawMode();
 							if (EntryPosesDrawMode == EMultiOptionDrawMode::All || (EntryPosesDrawMode == EMultiOptionDrawMode::Single && bIsActiveSet))
 							{
-								if (AnimTrack.Role != PrimaryRole && AnimTrack.Animation)
+								const UAnimSequenceBase* TrackAnimationSequence = Binding->GetAnimation();
+								if (AnimTrack.Role != PrimaryRole && TrackAnimationSequence)
 								{
 									const float Time = 0.f;
 									const FTransform RelativeToPrimary = SceneAsset->GetAlignmentTransformForRoleRelativeToOtherRole(AnimTrack.SectionIdx, AnimTrack.AnimSetIdx, AnimTrack.Role, PrimaryRole, Time);
@@ -207,7 +208,7 @@ void FContextualAnimEdMode::Render(const FSceneView* View, FViewport* Viewport,
 										Transform.SetLocation(Transform.GetLocation() - (FVector::UpVector * CapsuleHalfHeight));
 									}
 
-									UContextualAnimUtilities::DrawPose(PDI, AnimTrack.Animation, Time, Transform, bPassesSelectionCriteria ? FColor::Green : FColor::White, 1.f);
+									UContextualAnimUtilities::DrawPose(PDI, TrackAnimationSequence, Time, Transform, bPassesSelectionCriteria ? FColor::Green : FColor::White, 1.f);
 								}
 							}
 						}
diff --git a/Engine/Plugins/Experimental/Animation/ContextualAnimation/Source/ContextualAnimationEditor/Private/ContextualAnimViewModel.cpp b/Engine/Plugins/Experimental/Animation/ContextualAnimation/Source/ContextualAnimationEditor/Private/ContextualAnimViewModel.cpp
--- a/Engine/Plugins/Experimental/Animation/ContextualAnimation/Source/ContextualAnimationEditor/Private/ContextualAnimViewModel.cpp
+++ b/Engine/Plugins/Experimental/Animation/ContextualAnimation/Source/ContextualAnimationEditor/Private/ContextualAnimViewModel.cpp
@@ -8,6 +8,7 @@
 #include "ContextualAnimMovieSceneTrack.h"
 #include "ContextualAnimMovieSceneSection.h"
 #include "ContextualAnimMovieSceneNotifyTrack.h"
+#include "ContextualAnimOverrideInterface.h"
 #include "ContextualAnimPreviewScene.h"
 #include "Animation/AnimMontage.h"
 #include "EngineUtils.h"
@@ -575,6 +576,11 @@ void FContextualAnimViewModel::RefreshPreviewScene()
 		}
 	}
 
+	if (SceneAsset->PreviewOverrideProvider)
+	{
+		SceneBindings.SetOverrideProvider(TScriptInterface<IContextualAnimOverrideInterface>(SceneAsset->PreviewOverrideProvider->GetDefaultObject()));
+	}
+
 	// Pause all anims and disable auto blend out
 	const FContextualAnimSceneBinding* Leader = SceneBindings.GetSyncLeader();
 	if (Leader && Leader->GetSceneActorComponent()->StartContextualAnimScene(SceneBindings, {}))
@@ -827,7 +833,7 @@ void FContextualAnimViewModel::UpdatePreviewActorTransform(const FContextualAnim
 		// This allows the user to scrub the time line while maintaining the new MeshToScene transform 
 		if(ModifyingActorTransformInSceneState != EModifyActorTransformInSceneState::Inactive && Binding.GetActor() == ModifyingTransformInSceneCachedActor.Get())
 		{
-			const UAnimSequenceBase* Animation = SceneBindings.GetAnimTrackFromBinding(Binding).Animation;
+			const UAnimSequenceBase* Animation = Binding.GetAnimation();
 			if (Animation)
 			{
 				Transform = UContextualAnimUtilities::ExtractRootTransformFromAnimation(Animation, Time);
@@ -837,7 +843,15 @@ void FContextualAnimViewModel::UpdatePreviewActorTransform(const FContextualAnim
 		}
 		else
 		{
-			Transform = SceneBindings.GetAnimTrackFromBinding(Binding).GetRootTransformAtTime(Time);
+			const UAnimSequenceBase* AnimationSequence = Binding.GetAnimation();
+
+			FTransform RootTransform = FTransform::Identity;
+			if (AnimationSequence)
+			{
+				RootTransform = UContextualAnimUtilities::ExtractRootTransformFromAnimation(AnimationSequence, Time);
+			}
+
+			Transform = RootTransform * Binding.GetMeshToScene();
 		}
 
 		// Special case for Character
