[CAS] Remove Unused CAS Type For Warping

#rnx

diff --git a/Engine/Plugins/Experimental/Animation/ContextualAnimation/Source/ContextualAnimation/Public/ContextualAnimTypes.h b/Engine/Plugins/Experimental/Animation/ContextualAnimation/Source/ContextualAnimation/Public/ContextualAnimTypes.h
--- a/Engine/Plugins/Experimental/Animation/ContextualAnimation/Source/ContextualAnimation/Public/ContextualAnimTypes.h
+++ b/Engine/Plugins/Experimental/Animation/ContextualAnimation/Source/ContextualAnimation/Public/ContextualAnimTypes.h
@@ -538,10 +538,6 @@ private:
 	UPROPERTY()
 	FTransform MeshToScene;
 
-	/** Map of WarpTargetNames and Transforms for this set. Generated off line based on warp points defined in the asset */
-	UPROPERTY()
-	TMap<FName, FTransform> WarpPoints;
-
 	friend FContextualAnimSceneBinding;
 };
 
