Contextual Anim - Prevent nullptr crash in FContextualAnimSceneBindings::BindActorToRole

diff --git a/Engine/Plugins/Experimental/Animation/ContextualAnimation/Source/ContextualAnimation/Private/ContextualAnimTypes.cpp b/Engine/Plugins/Experimental/Animation/ContextualAnimation/Source/ContextualAnimation/Private/ContextualAnimTypes.cpp
--- a/Engine/Plugins/Experimental/Animation/ContextualAnimation/Source/ContextualAnimation/Private/ContextualAnimTypes.cpp
+++ b/Engine/Plugins/Experimental/Animation/ContextualAnimation/Source/ContextualAnimation/Private/ContextualAnimTypes.cpp
@@ -510,6 +510,13 @@ const FContextualAnimSceneBinding* FContextualAnimSceneBindings::GetPrimaryBindi
 
 bool FContextualAnimSceneBindings::BindActorToRole(AActor& ActorRef, FName Role, int32 SectionIdx, int32 AnimSetIdx)
 {
+	if (SceneAsset == nullptr)
+	{
+		UE_LOG(LogContextualAnim, Warning, TEXT("FContextualAnimSceneBindings::BindActorToRole. Failed to bind Actor: '%s' (Net Role: %d) to Role: '%s'. Reason: Invalid SceneAsset"),
+			*GetNameSafe(&ActorRef), ActorRef.GetLocalRole(), *Role.ToString());
+		return false;
+	}
+
 	const UContextualAnimSceneActorComponent* Comp = ActorRef.FindComponentByClass<UContextualAnimSceneActorComponent>();
 	if (Comp == nullptr)
 	{
