Tentative fix for GC leak

diff --git a/Engine/Plugins/Experimental/Animation/ContextualAnimation/Source/ContextualAnimation/Private/ContextualAnimSceneActorComponent.cpp b/Engine/Plugins/Experimental/Animation/ContextualAnimation/Source/ContextualAnimation/Private/ContextualAnimSceneActorComponent.cpp
--- a/Engine/Plugins/Experimental/Animation/ContextualAnimation/Source/ContextualAnimation/Private/ContextualAnimSceneActorComponent.cpp
+++ b/Engine/Plugins/Experimental/Animation/ContextualAnimation/Source/ContextualAnimation/Private/ContextualAnimSceneActorComponent.cpp
@@ -31,6 +31,15 @@ TAutoConsoleVariable<int32> CVarContextualAnimIKDebug(TEXT("a.ContextualAnim.IK.
 TAutoConsoleVariable<float> CVarContextualAnimIKDrawDebugLifetime(TEXT("a.ContextualAnim.IK.DrawDebugLifetime"), 0, TEXT("Draw Debug Duration"));
 #endif
 
+namespace ContextualAnimSceneActorComponentCVars
+{
+	static bool bLeaveSceneOnEndPlay = true;
+	static FAutoConsoleVariableRef CVarLeaveSceneOnEndPlay(
+		TEXT("a.ContextualAnim.LeaveSceneOnEndPlay"), bLeaveSceneOnEndPlay,
+		TEXT("If true, will call LeaveScene on EndPlay of each actor if they're in an active scene."),
+		ECVF_Default);
+}
+
 void FContextualAnimRepData::IncrementRepCounter()
 {
 	static uint8 Counter = 0;
@@ -101,6 +110,11 @@ void UContextualAnimSceneActorComponent::GetLifetimeReplicatedProps(TArray< FLif
 
 void UContextualAnimSceneActorComponent::EndPlay(const EEndPlayReason::Type EndPlayReason)
 {
+	if (ContextualAnimSceneActorComponentCVars::bLeaveSceneOnEndPlay && IsInActiveScene())
+	{
+		LeaveScene();
+	}
+
 	if (OwnerAnimInstance.IsValid())
 	{
 		OwnerAnimInstance->OnMontageBlendingOut.RemoveDynamic(this, &UContextualAnimSceneActorComponent::OnMontageBlendingOut);
