[CAS][Animation] Add Array of Collision Behaviours
- Lets us Combine Multiple Behaviours to achieve the desired effect
- Adds an option for setting the Collision Profile Directly
- Convert Saved Singlular Version To Array Version On PostLoad

#Synced-CL 48912871

diff --git a/Engine/Plugins/Experimental/Animation/ContextualAnimation/Source/ContextualAnimation/Private/ContextualAnimSceneActorComponent.cpp b/Engine/Plugins/Experimental/Animation/ContextualAnimation/Source/ContextualAnimation/Private/ContextualAnimSceneActorComponent.cpp
--- a/Engine/Plugins/Experimental/Animation/ContextualAnimation/Source/ContextualAnimation/Private/ContextualAnimSceneActorComponent.cpp
+++ b/Engine/Plugins/Experimental/Animation/ContextualAnimation/Source/ContextualAnimation/Private/ContextualAnimSceneActorComponent.cpp
@@ -23,6 +23,7 @@
 #include "Engine/World.h"
 #include "SceneView.h"
 #include "TimerManager.h"
+#include "Engine/CollisionProfile.h"
 
 #include UE_INLINE_GENERATED_CPP_BY_NAME(ContextualAnimSceneActorComponent)
 
@@ -1008,28 +1009,39 @@ void UContextualAnimSceneActorComponent::SetCollisionState(const FContextualAnim
 {
 	if (const UContextualAnimSceneAsset* Asset = Bindings.GetSceneAsset())
 	{
-		const EContextualAnimCollisionBehavior CollisionBehavior = Asset->GetCollisionBehavior();
-		if (CollisionBehavior == EContextualAnimCollisionBehavior::IgnoreActorWhenMoving)
+		const TArray<FContextualAnimCollisionBehaviorParams>& CollisionBehaviorParams = Asset->GetCollisionBehaviors();
+		for (const FContextualAnimCollisionBehaviorParams& CollisionBehaviorParam : CollisionBehaviorParams)
 		{
-			SetIgnoreCollisionWithOtherActors(true);
-		}
-		else if (CollisionBehavior == EContextualAnimCollisionBehavior::IgnoreChannels)
-		{
-			if (UPrimitiveComponent* RootPrimitiveComponent = Cast<UPrimitiveComponent>(GetOwner()->GetRootComponent()))
+			if (CollisionBehaviorParam.CollisionBehavior == EContextualAnimCollisionBehavior::IgnoreActorWhenMoving)
+			{
+				SetIgnoreCollisionWithOtherActors(true);
+			}
+			else if (CollisionBehaviorParam.CollisionBehavior == EContextualAnimCollisionBehavior::IgnoreChannels)
 			{
-				const TArray<TEnumAsByte<ECollisionChannel>>& ChannelsToIgnore = Asset->GetCollisionChannelsToIgnoreForRole(Bindings.GetRoleFromBinding(Binding));
-				if (ChannelsToIgnore.Num() > 0)
+				if (UPrimitiveComponent* RootPrimitiveComponent = Cast<UPrimitiveComponent>(GetOwner()->GetRootComponent()))
 				{
-					CharacterPropertiesBackup.CollisionResponses.Reset(ChannelsToIgnore.Num());
-					for (ECollisionChannel Channel : ChannelsToIgnore)
+					const TArray<TEnumAsByte<ECollisionChannel>>& ChannelsToIgnore = CollisionBehaviorParam.GetCollisionChannelsToIgnoreForRole(Bindings.GetRoleFromBinding(Binding));
+					if (ChannelsToIgnore.Num() > 0)
 					{
-						ECollisionResponse Response = RootPrimitiveComponent->GetCollisionResponseToChannel(Channel);
-						CharacterPropertiesBackup.CollisionResponses.Add(MakeTuple(Channel, Response));
+						CharacterPropertiesBackup.CollisionResponses.Reset(ChannelsToIgnore.Num());
+						for (ECollisionChannel Channel : ChannelsToIgnore)
+						{
+							ECollisionResponse Response = RootPrimitiveComponent->GetCollisionResponseToChannel(Channel);
+							CharacterPropertiesBackup.CollisionResponses.Add(MakeTuple(Channel, Response));
 
-						RootPrimitiveComponent->SetCollisionResponseToChannel(Channel, ECR_Ignore);
+							RootPrimitiveComponent->SetCollisionResponseToChannel(Channel, ECR_Ignore);
+						}
 					}
 				}
 			}
+			else if (CollisionBehaviorParam.CollisionBehavior == EContextualAnimCollisionBehavior::CollisionProfile)
+			{
+				if (UPrimitiveComponent* RootPrimitiveComponent = Cast<UPrimitiveComponent>(GetOwner()->GetRootComponent()))
+				{
+					CharacterPropertiesBackup.CollisionProfileResponses = RootPrimitiveComponent->GetCollisionResponseToChannels();
+					RootPrimitiveComponent->SetCollisionProfileName(CollisionBehaviorParam.GetCollisionProfileNameForRole(Bindings.GetRoleFromBinding(Binding)).Name);
+				}
+			}
 		}
 	}
 }
@@ -1038,21 +1050,35 @@ void UContextualAnimSceneActorComponent::RestoreCollisionState(const FContextual
 {
 	if (const UContextualAnimSceneAsset* Asset = Bindings.GetSceneAsset())
 	{
-		const EContextualAnimCollisionBehavior CollisionBehavior = Asset->GetCollisionBehavior();
-		if (CollisionBehavior == EContextualAnimCollisionBehavior::IgnoreActorWhenMoving)
+		const TArray<FContextualAnimCollisionBehaviorParams>& CollisionBehaviorParams = Asset->GetCollisionBehaviors();
+		for (const FContextualAnimCollisionBehaviorParams& CollisionBehaviorParam : CollisionBehaviorParams)
 		{
-			SetIgnoreCollisionWithOtherActors(false);
-		}
-		else if (CollisionBehavior == EContextualAnimCollisionBehavior::IgnoreChannels)
-		{
-			if (UPrimitiveComponent* RootPrimitiveComponent = Cast<UPrimitiveComponent>(GetOwner()->GetRootComponent()))
+			if (CollisionBehaviorParam.CollisionBehavior == EContextualAnimCollisionBehavior::IgnoreActorWhenMoving)
+			{
+				SetIgnoreCollisionWithOtherActors(false);
+			}
+			else if (CollisionBehaviorParam.CollisionBehavior == EContextualAnimCollisionBehavior::IgnoreChannels)
 			{
-				for (const TTuple<ECollisionChannel, ECollisionResponse>& Response : CharacterPropertiesBackup.CollisionResponses)
+				if (UPrimitiveComponent* RootPrimitiveComponent = Cast<UPrimitiveComponent>(GetOwner()->GetRootComponent()))
 				{
-					RootPrimitiveComponent->SetCollisionResponseToChannel(Response.Get<0>(), Response.Get<1>());
-				}
+					for (const TTuple<ECollisionChannel, ECollisionResponse>& Response : CharacterPropertiesBackup.CollisionResponses)
+					{
+						RootPrimitiveComponent->SetCollisionResponseToChannel(Response.Get<0>(), Response.Get<1>());
+					}
 
-				CharacterPropertiesBackup.CollisionResponses.Reset();
+					CharacterPropertiesBackup.CollisionResponses.Reset();
+				}
+			}
+			else if (CollisionBehaviorParam.CollisionBehavior == EContextualAnimCollisionBehavior::CollisionProfile)
+			{
+				if (UPrimitiveComponent* RootPrimitiveComponent = Cast<UPrimitiveComponent>(GetOwner()->GetRootComponent()))
+				{
+					if (CharacterPropertiesBackup.CollisionProfileResponses.IsSet())
+					{
+						RootPrimitiveComponent->SetCollisionResponseToChannels(CharacterPropertiesBackup.CollisionProfileResponses.GetValue());
+						CharacterPropertiesBackup.CollisionProfileResponses.Reset();
+					}
+				}
 			}
 		}
 	}
@@ -1184,10 +1210,13 @@ void UContextualAnimSceneActorComponent::OtherActorLeftScene(AActor& Actor)
 	{
 		if (const UContextualAnimSceneAsset* Asset = Bindings.GetSceneAsset())
 		{
-			const EContextualAnimCollisionBehavior CollisionBehavior = Asset->GetCollisionBehavior();
-			if (CollisionBehavior == EContextualAnimCollisionBehavior::IgnoreActorWhenMoving)
+			const TArray<FContextualAnimCollisionBehaviorParams>& CollisionBehaviorParams = Asset->GetCollisionBehaviors();
+			for (const FContextualAnimCollisionBehaviorParams& CollisionBehaviorParam : CollisionBehaviorParams)
 			{
-				SetIgnoreCollisionWithActor(Actor, false);
+				if (CollisionBehaviorParam.CollisionBehavior == EContextualAnimCollisionBehavior::IgnoreActorWhenMoving)
+				{
+					SetIgnoreCollisionWithActor(Actor, false);
+				}
 			}
 
 			Bindings.RemoveActor(Actor);
diff --git a/Engine/Plugins/Experimental/Animation/ContextualAnimation/Source/ContextualAnimation/Private/ContextualAnimSceneAsset.cpp b/Engine/Plugins/Experimental/Animation/ContextualAnimation/Source/ContextualAnimation/Private/ContextualAnimSceneAsset.cpp
--- a/Engine/Plugins/Experimental/Animation/ContextualAnimation/Source/ContextualAnimation/Private/ContextualAnimSceneAsset.cpp
+++ b/Engine/Plugins/Experimental/Animation/ContextualAnimation/Source/ContextualAnimation/Private/ContextualAnimSceneAsset.cpp
@@ -414,6 +414,36 @@ UContextualAnimSceneAsset::UContextualAnimSceneAsset(const FObjectInitializer& O
 	SampleRate = 15;
 }
 
+void UContextualAnimSceneAsset::PostLoad()
+{
+	Super::PostLoad();
+
+
+	// Only do this during editor time because that is when we can fix it up.
+// This means that properties will be fixed up during a cook!
+#if WITH_EDITORONLY_DATA
+	PRAGMA_DISABLE_DEPRECATION_WARNINGS
+
+		if (CollisionBehavior_DEPRECATED != EContextualAnimCollisionBehavior::None)
+		{
+			UE_LOG(LogContextualAnim, Log, TEXT("UContextualAnimSceneAsset::PostLoad - Collision Behavior Deprecation, Adding Entry For %s (%s)"), *GetNameSafe(this), *UEnum::GetValueAsString(CollisionBehavior_DEPRECATED));
+
+			FContextualAnimCollisionBehaviorParams BehaviourParams;
+			BehaviourParams.CollisionBehavior = CollisionBehavior_DEPRECATED;
+			BehaviourParams.CollisionChannelsToIgnoreParams = CollisionChannelsToIgnoreParams_DEPRECATED;
+
+			CollisionBehaviors.Insert(BehaviourParams, 0);
+
+			CollisionBehavior_DEPRECATED = EContextualAnimCollisionBehavior::None;
+			CollisionChannelsToIgnoreParams_DEPRECATED.Empty();
+		}
+
+	PRAGMA_ENABLE_DEPRECATION_WARNINGS
+#endif // WITH_EDITORONLY_DATA
+
+	
+}
+
 void UContextualAnimSceneAsset::PreSave(FObjectPreSaveContext ObjectSaveContext)
 {
 	Super::PreSave(ObjectSaveContext);
@@ -930,20 +960,6 @@ void UContextualAnimSceneAsset::GetAlignmentPointsForSecondaryRoleConsideringSel
 	}
 }
 
-const TArray<TEnumAsByte<ECollisionChannel>>& UContextualAnimSceneAsset::GetCollisionChannelsToIgnoreForRole(FName Role) const
-{
-	if (CollisionBehavior == EContextualAnimCollisionBehavior::IgnoreChannels)
-	{
-		if (const FContextualAnimIgnoreChannelsParam* Param = CollisionChannelsToIgnoreParams.FindByPredicate([Role](const FContextualAnimIgnoreChannelsParam& Item){ return Item.Role == Role; }))
-		{
-			return Param->Channels;
-		}
-	}
-
-	static const TArray<TEnumAsByte<ECollisionChannel>> EmptyArray;
-	return EmptyArray;
-}
-
 // Blueprint Interface
 //------------------------------------------------------------------------------------------
 
diff --git a/Engine/Plugins/Experimental/Animation/ContextualAnimation/Source/ContextualAnimation/Public/ContextualAnimSceneActorComponent.h b/Engine/Plugins/Experimental/Animation/ContextualAnimation/Source/ContextualAnimation/Public/ContextualAnimSceneActorComponent.h
--- a/Engine/Plugins/Experimental/Animation/ContextualAnimation/Source/ContextualAnimation/Public/ContextualAnimSceneActorComponent.h
+++ b/Engine/Plugins/Experimental/Animation/ContextualAnimation/Source/ContextualAnimation/Public/ContextualAnimSceneActorComponent.h
@@ -2,9 +2,10 @@
 
 #pragma once
 
-#include "ContextualAnimTypes.h"
 #include "ActorComponents/IKRigInterface.h"
 #include "Components/ActorComponent.h"
+#include "ContextualAnimTypes.h"
+
 #include "ContextualAnimSceneActorComponent.generated.h"
 
 #define UE_API CONTEXTUALANIMATION_API
@@ -292,6 +293,7 @@ protected:
 		TOptional<EMovementMode> MovementMode;
 		uint8 CustomMovementMode = 0;
 		TArray<TTuple<ECollisionChannel, ECollisionResponse>> CollisionResponses;
+		TOptional<struct FCollisionResponseContainer> CollisionProfileResponses;
 	};
 	FCharacterProperties CharacterPropertiesBackup;
 
diff --git a/Engine/Plugins/Experimental/Animation/ContextualAnimation/Source/ContextualAnimation/Public/ContextualAnimSceneAsset.h b/Engine/Plugins/Experimental/Animation/ContextualAnimation/Source/ContextualAnimation/Public/ContextualAnimSceneAsset.h
--- a/Engine/Plugins/Experimental/Animation/ContextualAnimation/Source/ContextualAnimation/Public/ContextualAnimSceneAsset.h
+++ b/Engine/Plugins/Experimental/Animation/ContextualAnimation/Source/ContextualAnimation/Public/ContextualAnimSceneAsset.h
@@ -2,6 +2,7 @@
 
 #pragma once
 
+#include "Engine/CollisionProfile.h"
 #include "Engine/DataAsset.h"
 #include "ContextualAnimTypes.h"
 #include "ContextualAnimSceneAsset.generated.h"
@@ -16,7 +17,8 @@ enum class EContextualAnimCollisionBehavior : uint8
 {
 	None,
 	IgnoreActorWhenMoving,
-	IgnoreChannels
+	IgnoreChannels,
+	CollisionProfile
 };
 
 USTRUCT(BlueprintType)
@@ -31,6 +33,61 @@ struct FContextualAnimIgnoreChannelsParam
 	TArray<TEnumAsByte<ECollisionChannel>> Channels;
 };
 
+USTRUCT(BlueprintType)
+struct FContextualAnimCollisionProfileParam
+{
+	GENERATED_BODY()
+
+	UPROPERTY(EditAnywhere, Category = "Defaults", meta = (GetOptions = "GetRoles"))
+	FName Role = NAME_None;
+
+	UPROPERTY(EditAnywhere, Category = "Defaults")
+	FCollisionProfileName Profile;
+};
+
+USTRUCT(BlueprintType)
+struct FContextualAnimCollisionBehaviorParams
+{
+	GENERATED_BODY()
+
+	UPROPERTY(EditAnywhere, Category = "Settings")
+	EContextualAnimCollisionBehavior CollisionBehavior = EContextualAnimCollisionBehavior::None;
+
+	UPROPERTY(EditAnywhere, Category = "Settings", meta = (EditCondition = "CollisionBehavior==EContextualAnimCollisionBehavior::IgnoreChannels", EditConditionHides))
+	TArray<FContextualAnimIgnoreChannelsParam> CollisionChannelsToIgnoreParams;
+
+	UPROPERTY(EditAnywhere, Category = "Settings", meta = (EditCondition = "CollisionBehavior==EContextualAnimCollisionBehavior::CollisionProfile", EditConditionHides))
+	TArray<FContextualAnimCollisionProfileParam> CollisionProfileParams;
+
+	const TArray<TEnumAsByte<ECollisionChannel>>& GetCollisionChannelsToIgnoreForRole(FName Role) const
+	{
+		if (CollisionBehavior == EContextualAnimCollisionBehavior::IgnoreChannels)
+		{
+			if (const FContextualAnimIgnoreChannelsParam* Param = CollisionChannelsToIgnoreParams.FindByPredicate([Role](const FContextualAnimIgnoreChannelsParam& Item) { return Item.Role == Role; }))
+			{
+				return Param->Channels;
+			}
+		}
+
+		static const TArray<TEnumAsByte<ECollisionChannel>> EmptyArray;
+		return EmptyArray;
+	}
+
+	FCollisionProfileName GetCollisionProfileNameForRole(FName Role) const
+	{
+		if (CollisionBehavior == EContextualAnimCollisionBehavior::CollisionProfile)
+		{
+			if (const FContextualAnimCollisionProfileParam* Param = CollisionProfileParams.FindByPredicate([Role](const FContextualAnimCollisionProfileParam& Item) { return Item.Role == Role; }))
+			{
+				return Param->Profile;
+			}
+		}
+
+		return FCollisionProfileName();
+	}
+
+};
+
 USTRUCT()
 struct FContextualAnimAttachmentParams
 {
@@ -227,6 +284,9 @@ public:
 	typedef TFunctionRef<UE::ContextualAnim::EForEachResult(const FContextualAnimTrack& AnimTrack)> FForEachAnimTrackFunction;
 
 	UE_API UContextualAnimSceneAsset(const FObjectInitializer& ObjectInitializer);
+	
+	/** Migrate deprecated objects. */
+	UE_API virtual void PostLoad() override;
 
 	UE_API virtual void PreSave(FObjectPreSaveContext ObjectSaveContext) override;
 
@@ -235,15 +295,16 @@ public:
 	UE_API void ForEachAnimTrack(FForEachAnimTrackFunction Function) const;
 
 	inline const FName& GetPrimaryRole() const { return PrimaryRole; }
-	inline EContextualAnimCollisionBehavior GetCollisionBehavior() const { return CollisionBehavior; }
+	UE_DEPRECATED(5.9, "Direct access has been deprecated, Please use CollisionBehaviors array version instead")
+	inline EContextualAnimCollisionBehavior GetCollisionBehavior() const { return CollisionBehaviors.Num() > 0 ? CollisionBehaviors[0].CollisionBehavior : EContextualAnimCollisionBehavior::None; }
+
+	inline const TArray<FContextualAnimCollisionBehaviorParams>&  GetCollisionBehaviors() const { return CollisionBehaviors; }
 	inline int32 GetSampleRate() const { return SampleRate; }
 	inline float GetRadius() const { return Radius; }
 	inline bool ShouldPrecomputeAlignmentTracks() const { return bPrecomputeAlignmentTracks; }
 	inline bool ShouldIgnoreClientMovementErrorChecksAndCorrection() const { return bIgnoreClientMovementErrorChecksAndCorrection; }
 	inline bool ShouldDisableMovementReplicationForSimulatedProxy() const { return bDisableMovementReplicationForSimulatedProxy; }
-
-	UE_API const TArray<TEnumAsByte<ECollisionChannel>>& GetCollisionChannelsToIgnoreForRole(FName Role) const;
-	
+		
 	const TArray<FContextualAnimAttachmentParams>& GetAttachmentParams() const { return AttachmentParams; }
 
 	const FContextualAnimAttachmentParams* GetAttachmentParamsForRole(FName Role) const
@@ -358,11 +419,16 @@ protected:
 	UPROPERTY(EditAnywhere, Category = "Settings")
 	float Radius = 0.f;
 
-	UPROPERTY(EditAnywhere, Category = "Settings")
-	EContextualAnimCollisionBehavior CollisionBehavior = EContextualAnimCollisionBehavior::None;
+#if WITH_EDITORONLY_DATA
+	UPROPERTY(meta = (DeprecatedProperty, DeprecationMessage = "Direct access has been deprecated, Please use CollisionBehaviors array version instead"))
+	EContextualAnimCollisionBehavior CollisionBehavior_DEPRECATED = EContextualAnimCollisionBehavior::None;
 
-	UPROPERTY(EditAnywhere, Category = "Settings", meta = (EditCondition = "CollisionBehavior==EContextualAnimCollisionBehavior::IgnoreChannels", EditConditionHides))
-	TArray<FContextualAnimIgnoreChannelsParam> CollisionChannelsToIgnoreParams;
+	UPROPERTY(meta = (DeprecatedProperty, DeprecationMessage = "Direct access has been deprecated, Please use CollisionBehaviors array version instead"))
+	TArray<FContextualAnimIgnoreChannelsParam> CollisionChannelsToIgnoreParams_DEPRECATED;
+#endif  // WITH_EDITORONLY_DATA
+
+	UPROPERTY(EditAnywhere, Category = "Settings")
+	TArray<FContextualAnimCollisionBehaviorParams> CollisionBehaviors;
 
 	UPROPERTY(EditAnywhere, Category = "Settings")
 	TArray<FContextualAnimAttachmentParams> AttachmentParams;
