[CAS] Fix Missing Animations in CAS When Transitioning Actors
- When Bindings are updated while Transitioning the animation can be missed
- Since the New API Asks the Binding for the Animation instead ofthe Track we need to make sure we update the Data in the Binding as well
- This also helps support the overrides efforts as now addign extra data when setting an anim track will ensure the binding stays in sync.

#Synced-CL 49122934

diff --git a/Engine/Plugins/Experimental/Animation/ContextualAnimation/Source/ContextualAnimation/Private/ContextualAnimSceneActorComponent.cpp b/Engine/Plugins/Experimental/Animation/ContextualAnimation/Source/ContextualAnimation/Private/ContextualAnimSceneActorComponent.cpp
--- a/Engine/Plugins/Experimental/Animation/ContextualAnimation/Source/ContextualAnimation/Private/ContextualAnimSceneActorComponent.cpp
+++ b/Engine/Plugins/Experimental/Animation/ContextualAnimation/Source/ContextualAnimation/Private/ContextualAnimSceneActorComponent.cpp
@@ -603,15 +603,19 @@ bool UContextualAnimSceneActorComponent::TransitionSingleActor(int32 SectionIdx,
 		return false;
 	}
 
-	if (const FContextualAnimSceneBinding* OwnerBinding = Bindings.FindBindingByActor(GetOwner()))
+	if (FContextualAnimSceneBinding* OwnerBinding = Bindings.FindBindingByActor(GetOwner()))
 	{
 		if (const UContextualAnimSceneAsset* Asset = Bindings.GetSceneAsset())
 		{
 			const FContextualAnimTrack* AnimTrack = Asset->GetAnimTrack(SectionIdx, AnimSetIdx, Bindings.GetRoleFromBinding(*OwnerBinding));
 			if (AnimTrack && AnimTrack->Animation)
 			{
+			
 				UE_LOG(LogContextualAnim, Verbose, TEXT("%-21s UContextualAnimSceneActorComponent::TransitionSingleActor Actor: %s SectionIdx: %d AnimSetIdx: %d"),
-					*UEnum::GetValueAsString(TEXT("Engine.ENetRole"), GetOwner()->GetLocalRole()), *GetNameSafe(GetOwner()), SectionIdx, AnimSetIdx);
+				*UEnum::GetValueAsString(TEXT("Engine.ENetRole"), GetOwner()->GetLocalRole()), *GetNameSafe(GetOwner()), SectionIdx, AnimSetIdx);
+
+				OwnerBinding->SetAnimTrack(*AnimTrack);
+				Bindings.UpdateBindingOverridesIfNeeded();
 
 				// Calculate WarpPoints
 				TArray<FContextualAnimWarpPoint> WarpPoints;
@@ -669,15 +673,18 @@ void UContextualAnimSceneActorComponent::OnRep_RepTransitionSingleActor()
 		return;
 	}
 
-	if (const FContextualAnimSceneBinding* OwnerBinding = Bindings.FindBindingByActor(GetOwner()))
+	if (FContextualAnimSceneBinding* OwnerBinding = Bindings.FindBindingByActor(GetOwner()))
 	{
 		if (RepTransitionSingleActorData.SectionIdx != MAX_uint8 && RepTransitionSingleActorData.AnimSetIdx != MAX_uint8)
 		{
 			if (const UContextualAnimSceneAsset* Asset = Bindings.GetSceneAsset())
 			{
 				const FContextualAnimTrack* AnimTrack = Asset->GetAnimTrack(RepTransitionSingleActorData.SectionIdx, RepTransitionSingleActorData.AnimSetIdx, Bindings.GetRoleFromBinding(*OwnerBinding));
-				if (AnimTrack && AnimTrack->Animation)
+				if (AnimTrack)
 				{
+					OwnerBinding->SetAnimTrack(*AnimTrack);
+					Bindings.UpdateBindingOverridesIfNeeded();
+
 					UAnimSequenceBase* TrackAnimationSequence = OwnerBinding->GetAnimation();
 					PlayAnimation_Internal(TrackAnimationSequence, 0.f, false);
 
diff --git a/Engine/Plugins/Experimental/Animation/ContextualAnimation/Source/ContextualAnimation/Private/ContextualAnimTypes.cpp b/Engine/Plugins/Experimental/Animation/ContextualAnimation/Source/ContextualAnimation/Private/ContextualAnimTypes.cpp
--- a/Engine/Plugins/Experimental/Animation/ContextualAnimation/Source/ContextualAnimation/Private/ContextualAnimTypes.cpp
+++ b/Engine/Plugins/Experimental/Animation/ContextualAnimation/Source/ContextualAnimation/Private/ContextualAnimTypes.cpp
@@ -1052,9 +1052,7 @@ void FContextualAnimSceneBindings::TransitionTo(int32 NewSectionIdx, int32 NewAn
 		{
 			if (const FContextualAnimTrack* AnimTrack = SceneAsset->GetAnimTrack(NewSectionIdx, NewAnimSetIdx, GetRoleFromBinding(Binding)))
 			{
-				Binding.SectionIdx = AnimTrack->SectionIdx;
-				Binding.AnimSetIdx = AnimTrack->AnimSetIdx;
-				Binding.AnimTrackIdx = AnimTrack->AnimTrackIdx;
+				Binding.SetAnimTrack(*AnimTrack);
 			}
 		}
 
diff --git a/Engine/Plugins/Experimental/Animation/ContextualAnimation/Source/ContextualAnimation/Public/ContextualAnimTypes.h b/Engine/Plugins/Experimental/Animation/ContextualAnimation/Source/ContextualAnimation/Public/ContextualAnimTypes.h
--- a/Engine/Plugins/Experimental/Animation/ContextualAnimation/Source/ContextualAnimation/Public/ContextualAnimTypes.h
+++ b/Engine/Plugins/Experimental/Animation/ContextualAnimation/Source/ContextualAnimation/Public/ContextualAnimTypes.h
@@ -558,6 +558,11 @@ struct FContextualAnimSceneBindings
 		return Actor ? Data.FindByPredicate([Actor](const FContextualAnimSceneBinding& Item) { return Item.GetActor() == Actor; }) : nullptr;
 	}
 
+	FContextualAnimSceneBinding* FindBindingByActor(const AActor* Actor)
+	{
+		return Actor ? Data.FindByPredicate([Actor](const FContextualAnimSceneBinding& Item) { return Item.GetActor() == Actor; }) : nullptr;
+	}
+
 	const FContextualAnimSceneBinding* FindBindingByRole(const FName& Role) const
 	{
 		return Role != NAME_None ? Data.FindByPredicate([this, &Role](const FContextualAnimSceneBinding& Item) { return GetAnimTrackFromBinding(Item).Role == Role; }) : nullptr;
