[CAS][Animation] Add Warnings Around Missing CASAC
- Relying on Spawned Component can be unsafe for replication this warning detects when the CAS Component doesnt exist while setting bindings.

#Synced-CL 49562015

diff --git a/Engine/Plugins/Experimental/Animation/ContextualAnimation/Source/ContextualAnimation/Private/ContextualAnimSceneActorComponent.cpp b/Engine/Plugins/Experimental/Animation/ContextualAnimation/Source/ContextualAnimation/Private/ContextualAnimSceneActorComponent.cpp
--- a/Engine/Plugins/Experimental/Animation/ContextualAnimation/Source/ContextualAnimation/Private/ContextualAnimSceneActorComponent.cpp
+++ b/Engine/Plugins/Experimental/Animation/ContextualAnimation/Source/ContextualAnimation/Private/ContextualAnimSceneActorComponent.cpp
@@ -950,7 +950,10 @@ void UContextualAnimSceneActorComponent::OnRep_Bindings()
 			{
 				if (Binding.GetActor() != GetOwner())
 				{
-					if (UContextualAnimSceneActorComponent* Comp = Binding.GetSceneActorComponent())
+					UContextualAnimSceneActorComponent* Comp = Binding.GetSceneActorComponent();
+
+					if (ensureMsgf(Comp, TEXT("%-21s UContextualAnimSceneActorComponent::OnRep_Bindings Actor: %s Current UContextualAnimSceneActorComponent INVALID"),
+						*UEnum::GetValueAsString(TEXT("Engine.ENetRole"), GetOwner()->GetLocalRole()), *GetNameSafe(Binding.GetActor())))
 					{
 						Comp->JoinScene(RepBindings.Bindings, RepBindings.WarpPoints, RepBindings.ExternalWarpTargets);
 					}
