Contextual Anim - Added GameplayTag to the scene asset that can be used as identifier for the asset
#Synced-CL 49618776

diff --git a/Engine/Plugins/Experimental/Animation/ContextualAnimation/Source/ContextualAnimation/Public/ContextualAnimSceneAsset.h b/Engine/Plugins/Experimental/Animation/ContextualAnimation/Source/ContextualAnimation/Public/ContextualAnimSceneAsset.h
--- a/Engine/Plugins/Experimental/Animation/ContextualAnimation/Source/ContextualAnimation/Public/ContextualAnimSceneAsset.h
+++ b/Engine/Plugins/Experimental/Animation/ContextualAnimation/Source/ContextualAnimation/Public/ContextualAnimSceneAsset.h
@@ -5,6 +5,7 @@
 #include "Engine/CollisionProfile.h"
 #include "Engine/DataAsset.h"
 #include "ContextualAnimTypes.h"
+#include "GameplayTagContainer.h"
 #include "ContextualAnimSceneAsset.generated.h"
 
 #define UE_API CONTEXTUALANIMATION_API
@@ -294,6 +295,7 @@ public:
 	
 	UE_API void ForEachAnimTrack(FForEachAnimTrackFunction Function) const;
 
+	inline const FGameplayTag& GetTag() const { return Tag; }
 	inline const FName& GetPrimaryRole() const { return PrimaryRole; }
 	UE_DEPRECATED(5.9, "Direct access has been deprecated, Please use CollisionBehaviors array version instead")
 	inline EContextualAnimCollisionBehavior GetCollisionBehavior() const { return CollisionBehaviors.Num() > 0 ? CollisionBehaviors[0].CollisionBehavior : EContextualAnimCollisionBehavior::None; }
@@ -403,6 +405,10 @@ protected:
 	UPROPERTY(EditAnywhere, Category = "Settings", meta = (GetOptions = "GetRoles"))
 	FName PrimaryRole = NAME_None;
 
+	/** Optional GameplayTag that can be used to identify this particular scene asset */
+	UPROPERTY(EditAnywhere, Category = "Settings")
+	FGameplayTag Tag;
+
 #if WITH_EDITORONLY_DATA
 
 	UPROPERTY(EditAnywhere, Category = "Settings", meta = (TitleProperty = "Role"))
