Contextual Anim Editor - Removed popup for confirming changes to MeshToScene property after moving an actor in the preview scene. Now MeshToScene changes as the user drags the actor around.

#Synced-CL 49618776

diff --git a/Engine/Plugins/Experimental/Animation/ContextualAnimation/Source/ContextualAnimation/Public/ContextualAnimActorInterface.h b/Engine/Plugins/Experimental/Animation/ContextualAnimation/Source/ContextualAnimation/Public/ContextualAnimActorInterface.h
--- a/Engine/Plugins/Experimental/Animation/ContextualAnimation/Source/ContextualAnimation/Public/ContextualAnimActorInterface.h
+++ b/Engine/Plugins/Experimental/Animation/ContextualAnimation/Source/ContextualAnimation/Public/ContextualAnimActorInterface.h
@@ -17,4 +17,8 @@ class IContextualAnimActorInterface : public IInterface
 
 	UFUNCTION(BlueprintNativeEvent, BlueprintCallable, Category = "Contextual Animation System")
 	class USkeletalMeshComponent* GetMesh() const;
+
+	/** Called from CAS Editor on the preview actor after spawn. Useful to give that actor a chance to run game-specific logic needed to be usable as preview actor in the CAS Editor */
+	UFUNCTION(BlueprintNativeEvent, BlueprintCallable, Category = "Contextual Animation System")
+	CONTEXTUALANIMATION_API void RefreshEditorPreview();
 };
diff --git a/Engine/Plugins/Experimental/Animation/ContextualAnimation/Source/ContextualAnimationEditor/Private/ContextualAnimViewModel.cpp b/Engine/Plugins/Experimental/Animation/ContextualAnimation/Source/ContextualAnimationEditor/Private/ContextualAnimViewModel.cpp
--- a/Engine/Plugins/Experimental/Animation/ContextualAnimation/Source/ContextualAnimationEditor/Private/ContextualAnimViewModel.cpp
+++ b/Engine/Plugins/Experimental/Animation/ContextualAnimation/Source/ContextualAnimationEditor/Private/ContextualAnimViewModel.cpp
@@ -9,8 +9,10 @@
 #include "ContextualAnimMovieSceneSection.h"
 #include "ContextualAnimMovieSceneNotifyTrack.h"
 #include "ContextualAnimOverrideInterface.h"
+#include "ContextualAnimActorInterface.h"
 #include "ContextualAnimPreviewScene.h"
 #include "Animation/AnimMontage.h"
+#include "Editor.h"
 #include "EngineUtils.h"
 #include "Framework/Application/SlateApplication.h"
 #include "ContextualAnimEditorTypes.h"
@@ -306,6 +308,12 @@ AActor* FContextualAnimViewModel::SpawnPreviewActor(const FContextualAnimTrack&
 			{
 				PreviewActor = GetWorld()->SpawnActor<AActor>(PreviewActorClass, SpawnTransform, Params);
 				PreviewActor->SetFlags(RF_Transient);
+
+				// Try to call RefreshEditorPreview on the actor. This give game specific actors a chance to run any logic needed to be able to use them as preview actor in this editor.
+				if (PreviewActor->GetClass()->ImplementsInterface(UContextualAnimActorInterface::StaticClass()))
+				{
+					IContextualAnimActorInterface::Execute_RefreshEditorPreview(PreviewActor);
+				}
 			}
 			else
 			{
@@ -827,32 +835,7 @@ void FContextualAnimViewModel::UpdatePreviewActorTransform(const FContextualAnim
 {
 	if (AActor* PreviewActor = Binding.GetActor())
 	{
-		FTransform Transform = FTransform::Identity;
-
-		// When modifying the actor transform, use the 'temp' transform instead of the one saved in the asset
-		// This allows the user to scrub the time line while maintaining the new MeshToScene transform 
-		if(ModifyingActorTransformInSceneState != EModifyActorTransformInSceneState::Inactive && Binding.GetActor() == ModifyingTransformInSceneCachedActor.Get())
-		{
-			const UAnimSequenceBase* Animation = Binding.GetAnimation();
-			if (Animation)
-			{
-				Transform = UContextualAnimUtilities::ExtractRootTransformFromAnimation(Animation, Time);
-			}
-
-			Transform *= NewMeshToSceneTransform;
-		}
-		else
-		{
-			const UAnimSequenceBase* AnimationSequence = Binding.GetAnimation();
-
-			FTransform RootTransform = FTransform::Identity;
-			if (AnimationSequence)
-			{
-				RootTransform = UContextualAnimUtilities::ExtractRootTransformFromAnimation(AnimationSequence, Time);
-			}
-
-			Transform = RootTransform * Binding.GetMeshToScene();
-		}
+		FTransform Transform = SceneBindings.GetAnimTrackFromBinding(Binding).GetRootTransformAtTime(Time);
 
 		// Special case for Character
 		if (ACharacter* PreviewCharacter = Cast<ACharacter>(PreviewActor))
@@ -1028,13 +1011,6 @@ void FContextualAnimViewModel::UpdateSelection(const AActor* SelectedActor)
 	const FContextualAnimSceneBinding* Binding = SceneBindings.FindBindingByActor(SelectedActor);
 	if (Binding)
 	{
-		// Discard changes if we were modifying the transform of an actor in the scene but never confirmed the changes before selecting a different actor
-		if(ModifyingActorTransformInSceneState == EModifyActorTransformInSceneState::WaitingForConfirmation && 
-			Binding->GetActor() != ModifyingTransformInSceneCachedActor.Get())
-		{
-			DiscardChangeToActorTransformInScene();
-		}
-
 		UpdateSelection(SceneBindings.GetRoleFromBinding(*Binding));
 	}
 	else
@@ -1181,11 +1157,14 @@ bool FContextualAnimViewModel::ProcessInputDelta(FVector& InDrag, FRotator& InRo
 		{
 			if (WantsToModifyMeshToSceneForSelectedActor())
 			{
-				if (AActor* SelectedActor = GetSelectedActor())
+				AActor* SelectedActor = GetSelectedActor();
+				FContextualAnimTrack* SelectedTrack = GetSelectedAnimTrack();
+				if (SelectedActor && SelectedTrack)
 				{
-					NewMeshToSceneTransform.SetLocation(NewMeshToSceneTransform.GetLocation() + InDrag);
-					NewMeshToSceneTransform.SetRotation(FQuat(InRot) * NewMeshToSceneTransform.GetRotation());
-					SelectedActor->SetActorLocationAndRotation(SelectedActor->GetActorLocation() + InDrag, SelectedActor->GetActorRotation() + InRot);
+					SelectedTrack->MeshToScene.SetLocation(SelectedTrack->MeshToScene.GetLocation() + InDrag);
+					SelectedTrack->MeshToScene.SetRotation(FQuat(InRot) * SelectedTrack->MeshToScene.GetRotation());
+
+					SelectedActor->SetActorLocationAndRotation(SelectedActor->GetActorLocation() + InDrag, SelectedActor->GetActorRotation() + InRot);				
 				}
 			}
 
@@ -1283,9 +1262,11 @@ bool FContextualAnimViewModel::StartTracking()
 		{
 			if (ModifyingActorTransformInSceneState == EModifyActorTransformInSceneState::Inactive)
 			{
-				NewMeshToSceneTransform = SceneBindings.GetAnimTrackFromBinding(*Binding).MeshToScene;
 				ModifyingActorTransformInSceneState = EModifyActorTransformInSceneState::Modifying;
 				ModifyingTransformInSceneCachedActor = Binding->GetActor();
+
+				GEditor->BeginTransaction(LOCTEXT("ModifyMeshToSceneTransform", "Modify MeshToScene Transform"));
+				SceneAsset->Modify();
 			}
 
 			return true;
@@ -1299,16 +1280,10 @@ bool FContextualAnimViewModel::EndTracking()
 {
 	if (ModifyingActorTransformInSceneState == EModifyActorTransformInSceneState::Modifying)
 	{
-		const FContextualAnimSceneBinding* Binding = GetSelectedBinding();
-		if (Binding && !SceneBindings.GetAnimTrackFromBinding(*Binding).MeshToScene.Equals(NewMeshToSceneTransform))
-		{
-			ModifyingActorTransformInSceneState = EModifyActorTransformInSceneState::WaitingForConfirmation;
-		}
-		else
-		{
-			ModifyingActorTransformInSceneState = EModifyActorTransformInSceneState::Inactive;
-			ModifyingTransformInSceneCachedActor.Reset();
-		}
+		GEditor->EndTransaction();
+		
+		ModifyingActorTransformInSceneState = EModifyActorTransformInSceneState::Inactive;
+		ModifyingTransformInSceneCachedActor.Reset();
 
 		return true;		
 	}
@@ -1316,42 +1291,6 @@ bool FContextualAnimViewModel::EndTracking()
 	return false;
 }
 
-bool FContextualAnimViewModel::IsChangeToActorTransformInSceneWaitingForConfirmation() const
-{
-	const FContextualAnimSceneBinding* Binding = GetSelectedBinding();
-	return (Binding && Binding->GetActor() == ModifyingTransformInSceneCachedActor.Get() && ModifyingActorTransformInSceneState == EModifyActorTransformInSceneState::WaitingForConfirmation);
-}
-
-void FContextualAnimViewModel::ApplyChangeToActorTransformInScene()
-{
-	const FContextualAnimSceneBinding* Binding = GetSelectedBinding();
-	if(Binding && Binding->GetActor() == ModifyingTransformInSceneCachedActor.Get())
-	{
-		FScopedTransaction Transaction(LOCTEXT("ModifyMeshToSceneTransform", "Modify MeshToScene Transform"));
-		SceneAsset->Modify();
-
-		FTransform& MeshToSpaceTransform = (const_cast<FContextualAnimTrack*>(&SceneBindings.GetAnimTrackFromBinding(*Binding)))->MeshToScene;
-		MeshToSpaceTransform = NewMeshToSceneTransform;
-
-		UpdatePreviewActorTransform(*Binding, Sequencer->GetGlobalTime().AsSeconds());
-	}
-	
-	ModifyingActorTransformInSceneState = EModifyActorTransformInSceneState::Inactive;
-	ModifyingTransformInSceneCachedActor.Reset();
-}
-
-void FContextualAnimViewModel::DiscardChangeToActorTransformInScene()
-{
-	ModifyingActorTransformInSceneState = EModifyActorTransformInSceneState::Inactive;
-
-	if (const FContextualAnimSceneBinding* Binding = SceneBindings.FindBindingByActor(ModifyingTransformInSceneCachedActor.Get()))
-	{
-		UpdatePreviewActorTransform(*Binding, Sequencer->GetGlobalTime().AsSeconds());
-	}
-
-	ModifyingTransformInSceneCachedActor.Reset();
-}
-
 void FContextualAnimViewModel::UpdateRoles()
 {
 	FScopedTransaction Transaction(LOCTEXT("CAS_UpdateRoles", "UpdateRoles"));
diff --git a/Engine/Plugins/Experimental/Animation/ContextualAnimation/Source/ContextualAnimationEditor/Private/SContextualAnimViewport.cpp b/Engine/Plugins/Experimental/Animation/ContextualAnimation/Source/ContextualAnimationEditor/Private/SContextualAnimViewport.cpp
--- a/Engine/Plugins/Experimental/Animation/ContextualAnimation/Source/ContextualAnimationEditor/Private/SContextualAnimViewport.cpp
+++ b/Engine/Plugins/Experimental/Animation/ContextualAnimation/Source/ContextualAnimationEditor/Private/SContextualAnimViewport.cpp
@@ -35,18 +35,16 @@ void SContextualAnimViewport::Construct(const FArguments& InArgs, const FContext
 
 	auto GetNotificationVisibility = [this]()
 	{
-		return AssetEditorToolkitPtr.Pin()->GetViewModel()->IsChangeToActorTransformInSceneWaitingForConfirmation() ? EVisibility::Visible : EVisibility::Hidden;
+		return EVisibility::Hidden;
 	};
 
 	auto ApplyChanges = [this]()
 	{
-		AssetEditorToolkitPtr.Pin()->GetViewModel()->ApplyChangeToActorTransformInScene();
 		return FReply::Handled();
 	};
 
 	auto CancelChanges = [this]()
  	{
-		AssetEditorToolkitPtr.Pin()->GetViewModel()->DiscardChangeToActorTransformInScene();
 		return FReply::Handled();
 	};
 
diff --git a/Engine/Plugins/Experimental/Animation/ContextualAnimation/Source/ContextualAnimationEditor/Public/ContextualAnimViewModel.h b/Engine/Plugins/Experimental/Animation/ContextualAnimation/Source/ContextualAnimationEditor/Public/ContextualAnimViewModel.h
--- a/Engine/Plugins/Experimental/Animation/ContextualAnimation/Source/ContextualAnimationEditor/Public/ContextualAnimViewModel.h
+++ b/Engine/Plugins/Experimental/Animation/ContextualAnimation/Source/ContextualAnimationEditor/Public/ContextualAnimViewModel.h
@@ -122,10 +122,6 @@ public:
 	bool StartTracking();
 	bool EndTracking();
 
-	bool IsChangeToActorTransformInSceneWaitingForConfirmation() const;
-	void ApplyChangeToActorTransformInScene();
-	void DiscardChangeToActorTransformInScene();
-
 	float GetPlaybackTime() const;
 	
 	bool CanMakeEdits() const;
@@ -208,9 +204,6 @@ private:
 	enum class EModifyActorTransformInSceneState : uint8 {  Inactive, Modifying, WaitingForConfirmation };
 	EModifyActorTransformInSceneState ModifyingActorTransformInSceneState = EModifyActorTransformInSceneState::Inactive;
 
-	/** New MeshToScene value for the selected actor. Updated while dragging the actor around. Only committed when the user confirms the operation */
-	FTransform NewMeshToSceneTransform = FTransform::Identity;
-
 	/** Selected actor when the user starts modifying its transform in the scene */
 	TWeakObjectPtr<AActor> ModifyingTransformInSceneCachedActor;
 
