[CAS][Overrides] Allow ST Task to provide CAS Overrides

#Synced-CL 50355740

#Synced-CL 50355740

diff --git a/Engine/Plugins/Experimental/Animation/ContextualAnimation/Source/ContextualAnimation/Public/ContextualAnimOverrideInterface.h b/Engine/Plugins/Experimental/Animation/ContextualAnimation/Source/ContextualAnimation/Public/ContextualAnimOverrideInterface.h
--- a/Engine/Plugins/Experimental/Animation/ContextualAnimation/Source/ContextualAnimation/Public/ContextualAnimOverrideInterface.h
+++ b/Engine/Plugins/Experimental/Animation/ContextualAnimation/Source/ContextualAnimation/Public/ContextualAnimOverrideInterface.h
@@ -42,8 +42,7 @@ struct FContextualAnimationOverrideData
 	FContextualAnimAlignmentTrackContainer IKTargetData;
 };
 
-
-UINTERFACE()
+UINTERFACE(BlueprintType)
 class UContextualAnimOverrideInterface : public UInterface
 {
 	GENERATED_BODY()
