Contextual Anim - Added OnMontageReachedEndDelegate delegate called when the currently playing montage reaches the end but does not blends out

#Synced-CL 50533237

diff --git a/Engine/Plugins/Experimental/Animation/ContextualAnimation/Source/ContextualAnimation/Private/ContextualAnimSceneActorComponent.cpp b/Engine/Plugins/Experimental/Animation/ContextualAnimation/Source/ContextualAnimation/Private/ContextualAnimSceneActorComponent.cpp
--- a/Engine/Plugins/Experimental/Animation/ContextualAnimation/Source/ContextualAnimation/Private/ContextualAnimSceneActorComponent.cpp
+++ b/Engine/Plugins/Experimental/Animation/ContextualAnimation/Source/ContextualAnimation/Private/ContextualAnimSceneActorComponent.cpp
@@ -163,6 +163,8 @@ void UContextualAnimSceneActorComponent::PlayAnimation_Internal(UAnimSequenceBas
 			AnimInstance->OnMontageBlendingOut.AddUniqueDynamic(this, &UContextualAnimSceneActorComponent::OnMontageBlendingOut);
 			AnimInstance->OnPlayMontageNotifyBegin.AddUniqueDynamic(this, &UContextualAnimSceneActorComponent::OnPlayMontageNotifyBegin);
 
+			LoopCounter = 0;
+
 			//@TODO: Move this out of this function
 			if (bSyncPlaybackTime)
 			{
@@ -1434,6 +1436,39 @@ void UContextualAnimSceneActorComponent::OnPlayMontageNotifyBegin(FName NotifyNa
 
 void UContextualAnimSceneActorComponent::OnTickPose(class USkinnedMeshComponent* SkinnedMeshComponent, float DeltaTime, bool bNeedsValidRootMotion)
 {
+	// Fire OnMontageReachedEndDelegate if needed.
+	const FContextualAnimSceneBinding* BindingPtr = Bindings.IsValid() ? Bindings.FindBindingByActor(GetOwner()) : nullptr;
+	const FAnimMontageInstance* MontageInstance = BindingPtr ? BindingPtr->GetAnimMontageInstance() : nullptr;
+	if (MontageInstance && MontageInstance->Montage && MontageInstance->Montage == BindingPtr->GetAnimation())
+	{
+		const float MontagePos = MontageInstance->GetPosition();
+		const FName NextSection = MontageInstance->GetNextSection();
+
+		// EnableAutoBlendOut in the montage is set to false
+		if (NextSection == NAME_None && MontageInstance->Montage->bEnableAutoBlendOut == false)
+		{
+			float StartTime = 0.f;
+			float EndTime = 0.f;
+			MontageInstance->Montage->GetSectionStartAndEndTime(MontageInstance->Montage->GetSectionIndex(MontageInstance->GetCurrentSection()), StartTime, EndTime);
+
+			// Use LoopCounter in this case as a flag to ensure we fire OnMontageReachedEndDelegate only once
+			if (LoopCounter == 0 && FMath::IsNearlyEqual(MontagePos, EndTime, UE_KINDA_SMALL_NUMBER)) 
+			{
+				OnMontageReachedEndDelegate.Broadcast(this, MontageInstance->Montage, 0);
+				LoopCounter++;
+			}
+		}
+		// Section in the montage is set to loop
+		else if (NextSection == MontageInstance->GetCurrentSection())
+		{
+			if (MontagePos < MontageInstance->GetPreviousPosition())
+			{
+				LoopCounter++;
+				OnMontageReachedEndDelegate.Broadcast(this, MontageInstance->Montage, LoopCounter);
+			}
+		}
+	}
+
 	//@TODO: Check for LOD too to prevent this update if the actor is too far away
 	if (GetNetMode() != NM_DedicatedServer)
 	{
diff --git a/Engine/Plugins/Experimental/Animation/ContextualAnimation/Source/ContextualAnimation/Public/ContextualAnimSceneActorComponent.h b/Engine/Plugins/Experimental/Animation/ContextualAnimation/Source/ContextualAnimation/Public/ContextualAnimSceneActorComponent.h
--- a/Engine/Plugins/Experimental/Animation/ContextualAnimation/Source/ContextualAnimation/Public/ContextualAnimSceneActorComponent.h
+++ b/Engine/Plugins/Experimental/Animation/ContextualAnimation/Source/ContextualAnimation/Public/ContextualAnimSceneActorComponent.h
@@ -19,6 +19,7 @@ class UAnimSequenceBase;
 DECLARE_DYNAMIC_MULTICAST_DELEGATE_OneParam(FContextualAnimSceneActorCompDelegate, class UContextualAnimSceneActorComponent*, SceneActorComponent);
 DECLARE_DYNAMIC_MULTICAST_DELEGATE_TwoParams(FContextualAnimPlayMontageNotifyBeginDelegate, class UContextualAnimSceneActorComponent*, SceneActorComponent, FName, NotifyName);
 DECLARE_DYNAMIC_MULTICAST_DELEGATE_ThreeParams(FContextualAnimOnMontageBlendingOutDelegate, class UContextualAnimSceneActorComponent*, SceneComponent, UAnimMontage* const, BlendingOutMontage, bool, bInterrupted);
+DECLARE_DYNAMIC_MULTICAST_DELEGATE_ThreeParams(FContextualAnimOnMontageReachedEndDelegate, class UContextualAnimSceneActorComponent*, SceneComponent, UAnimMontage* const, Montage, int32, LoopCounter);
 
 USTRUCT(BlueprintType)
 struct FContextualAnimWarpTarget
@@ -172,6 +173,10 @@ public:
 	UPROPERTY(BlueprintAssignable, Category = "Events")
 	FContextualAnimOnMontageBlendingOutDelegate OnMontageBlendingOutDelegate;
 
+	/** Delegate called when the currently playing montage reaches the end but does not blends out */
+	UPROPERTY(BlueprintAssignable, Category = "Events")
+	FContextualAnimOnMontageReachedEndDelegate OnMontageReachedEndDelegate;
+
 	UPROPERTY(EditAnywhere, BlueprintReadWrite, Category = "Config")
 	TObjectPtr<class UContextualAnimSceneAsset> SceneAsset;
 
@@ -376,6 +381,7 @@ protected:
 private:
 
 	TArray<TWeakObjectPtr<const UAnimMontage>, TInlineAllocator<5>> AnimsPlayed;
+	int32 LoopCounter = 0;
 };
 
 #undef UE_API
