[CAS] Rework the CAS New Animset menu. Allow the ability to quickly create a new animset and add it to the Active section

#changelist validated
#Synced-CL 51625089

#Synced-CL 51625089

diff --git a/Engine/Plugins/Experimental/Animation/ContextualAnimation/Source/ContextualAnimationEditor/Private/ContextualAnimAssetEditorToolkit.cpp b/Engine/Plugins/Experimental/Animation/ContextualAnimation/Source/ContextualAnimationEditor/Private/ContextualAnimAssetEditorToolkit.cpp
--- a/Engine/Plugins/Experimental/Animation/ContextualAnimation/Source/ContextualAnimationEditor/Private/ContextualAnimAssetEditorToolkit.cpp
+++ b/Engine/Plugins/Experimental/Animation/ContextualAnimation/Source/ContextualAnimationEditor/Private/ContextualAnimAssetEditorToolkit.cpp
@@ -8,6 +8,7 @@
 #include "SContextualAnimNewAnimSetDialog.h"
 #include "ContextualAnimPreviewScene.h"
 #include "ContextualAnimAssetEditorCommands.h"
+#include "ContextualAnimEditorTypes.h"
 #include "Widgets/Docking/SDockTab.h"
 #include "ContextualAnimViewModel.h"
 #include "ContextualAnimSceneAsset.h"
@@ -206,12 +207,16 @@ void FContextualAnimAssetEditorToolkit::FillToolbar(FToolBarBuilder& ToolbarBuil
 		FText::GetEmpty(),
 		FSlateIcon()
 	);
-
-	ToolbarBuilder.AddToolBarButton(
-		FContextualAnimAssetEditorCommands::Get().NewAnimSet,
-		NAME_None,
-		TAttribute<FText>(),
-		TAttribute<FText>(),
+	
+	ToolbarBuilder.AddComboButton(
+	FUIAction(
+		FExecuteAction(),
+		FCanExecuteAction::CreateSP(this, &FContextualAnimAssetEditorToolkit::CanMakeEdits),
+		FIsActionChecked()
+	),
+		FOnGetContent::CreateSP(this, &FContextualAnimAssetEditorToolkit::BuildAnimSetMenu),
+		LOCTEXT("AnimSet_Label", "New AnimSet"),
+		FText::GetEmpty(),
 		FSlateIcon(FAppStyle::GetAppStyleSetName(), "Icons.Plus")
 	);
 
@@ -250,6 +255,32 @@ TSharedRef<SWidget> FContextualAnimAssetEditorToolkit::BuildSectionsMenu()
 	return MenuBuilder.MakeWidget();
 }
 
+TSharedRef<SWidget> FContextualAnimAssetEditorToolkit::BuildAnimSetMenu()
+{
+	const bool bShouldCloseWindowAfterMenuSelection = true;
+	FMenuBuilder MenuBuilder(bShouldCloseWindowAfterMenuSelection, GetToolkitCommands());
+
+	MenuBuilder.BeginSection(NAME_None, LOCTEXT("AnimSet_Label", "New AnimSet"));
+	{
+		MenuBuilder.AddMenuEntry(
+			FText::FromString("Add New AnimSet"),
+			FText::GetEmpty(),
+			FSlateIcon(),
+			FUIAction(FExecuteAction::CreateLambda([this]() {
+				ShowNewAnimSetDialog();
+			})));
+		MenuBuilder.AddMenuEntry(
+			FText::FromString("Add Empty AnimSet to Active Section"),
+			FText::GetEmpty(),
+			FSlateIcon(),
+			FUIAction(FExecuteAction::CreateLambda([this]() {AddAnimSetToActiveSection();
+			})));
+	}
+
+	MenuBuilder.EndSection();
+	return MenuBuilder.MakeWidget();
+}
+
 void FContextualAnimAssetEditorToolkit::ShowNewAnimSetDialog()
 {
 	const UContextualAnimSceneAsset* SceneAsset = GetSceneAsset();
@@ -266,6 +297,34 @@ void FContextualAnimAssetEditorToolkit::ShowNewAnimSetDialog()
 	NewAnimSetDialog->Show();
 }
 
+void FContextualAnimAssetEditorToolkit::AddAnimSetToActiveSection()
+{
+	const UContextualAnimSceneAsset* SceneAsset = GetSceneAsset();
+	const int32 NumRoles = SceneAsset ? SceneAsset->GetNumRoles() : 0;
+	
+
+	if (NumRoles == 0)
+	{
+		const FText DialogMsg = LOCTEXT("MissingRolesDialog", "Can't add new anim set. Roles asset hasn't been set or no roles has been defined.");
+		FMessageDialog::Open(EAppMsgType::Ok, DialogMsg);
+		return;
+	}
+	
+	FContextualAnimNewAnimSetParams Params = FContextualAnimNewAnimSetParams();
+	const TArray<FName> Roles = SceneAsset->GetRoles();
+	
+	for (FName RoleName: Roles)
+	{
+		FContextualAnimNewAnimSetData RoleData;
+		RoleData.RoleName = RoleName; 
+		Params.Data.Add(RoleData);
+	}
+	
+	const int32 ActiveSectionIdx = ViewModel->GetActiveSection();
+	Params.SectionName = SceneAsset->GetSection(ActiveSectionIdx)->GetName();
+	ViewModel->AddNewAnimSet(Params);
+}
+
 void FContextualAnimAssetEditorToolkit::RegisterTabSpawners(const TSharedRef<FTabManager>& InTabManager)
 {
 	WorkspaceMenuCategory = InTabManager->AddLocalWorkspaceMenuCategory(LOCTEXT("WorkspaceMenu_CASEditor", "Contextual Animation Editor"));
diff --git a/Engine/Plugins/Experimental/Animation/ContextualAnimation/Source/ContextualAnimationEditor/Public/ContextualAnimAssetEditorToolkit.h b/Engine/Plugins/Experimental/Animation/ContextualAnimation/Source/ContextualAnimationEditor/Public/ContextualAnimAssetEditorToolkit.h
--- a/Engine/Plugins/Experimental/Animation/ContextualAnimation/Source/ContextualAnimationEditor/Public/ContextualAnimAssetEditorToolkit.h
+++ b/Engine/Plugins/Experimental/Animation/ContextualAnimation/Source/ContextualAnimationEditor/Public/ContextualAnimAssetEditorToolkit.h
@@ -55,8 +55,10 @@ private:
 	void FillToolbar(FToolBarBuilder& ToolbarBuilder);
 
 	TSharedRef<SWidget> BuildSectionsMenu();
+	TSharedRef<SWidget> BuildAnimSetMenu();
 
 	void ShowNewAnimSetDialog();
+	void AddAnimSetToActiveSection();
 
 	TSharedPtr<SContextualAnimViewport> ViewportWidget;
 
