Contextual Anim - Added option to let the system find ideal start time for the interaction based on the distance from the actor to the warp target in the animation and the actual distance in game.
#Synced-CL 51627494

diff --git a/Engine/Plugins/Experimental/Animation/ContextualAnimation/Source/ContextualAnimation/Private/ContextualAnimSceneActorComponent.cpp b/Engine/Plugins/Experimental/Animation/ContextualAnimation/Source/ContextualAnimation/Private/ContextualAnimSceneActorComponent.cpp
--- a/Engine/Plugins/Experimental/Animation/ContextualAnimation/Source/ContextualAnimation/Private/ContextualAnimSceneActorComponent.cpp
+++ b/Engine/Plugins/Experimental/Animation/ContextualAnimation/Source/ContextualAnimation/Private/ContextualAnimSceneActorComponent.cpp
@@ -765,6 +765,74 @@ bool UContextualAnimSceneActorComponent::TransitionSingleActor(int32 SectionIdx,
 	return TransitionSingleActor(SectionIdx, AnimSetIdx, {});
 }
 
+float UContextualAnimSceneActorComponent::FindIdealStartTime(const FContextualAnimSceneBindings& InBindings, const FContextualAnimSceneBinding& Binding, const TArray<FContextualAnimWarpPoint>& WarpPoints) const
+{
+	float IdealStartTime = 0.f;
+
+	const UContextualAnimSceneAsset* Asset = InBindings.GetSceneAsset();
+	if (Asset && Binding.GetAnimation() && WarpPoints.Num() > 0)
+	{
+		// Get the end time for the first warp section in the animation
+		// This is our hard upper bound since we need to give some chance to motion warping to run
+		float WarpStartTime, WarpEndTime;
+		FName WarpTargetName = UContextualAnimUtilities::GetStartAndEndTimeForWarpSection(Binding.GetAnimation(), 0, WarpStartTime, WarpEndTime);
+
+		// Find the warp point that matches the warp target name from the animation
+		const FContextualAnimWarpPoint* WarpPoint = WarpPoints.FindByPredicate([WarpTargetName](const FContextualAnimWarpPoint& Item) { return Item.Name == WarpTargetName; });
+		if (WarpPoint == nullptr)
+		{
+			UE_LOG(LogContextualAnim, Verbose, TEXT("FindIdealStartTime: No warp point found for WarpTargetName: %s Actor: %s Animation: %s"),
+				*WarpTargetName.ToString(), *GetNameSafe(GetOwner()), *GetNameSafe(Binding.GetAnimation()));
+			return IdealStartTime;
+		}
+
+		// Early out if there is no enough room for motion warping to run
+		const float MaxSearchTime = WarpEndTime - Asset->GetFindIdealStartTimeMargin();
+		if (MaxSearchTime <= 0.f)
+		{
+			return IdealStartTime;
+		}
+
+		// Calculate where the animation expects the character to be at T=0 and at sync time, both relative to the warp point
+		const FVector AnimStartLocation = Asset->GetAlignmentTransform(Binding, InBindings, WarpTargetName, 0.f).GetLocation();
+		const FVector AnimSyncLocation = Asset->GetAlignmentTransform(Binding, InBindings, WarpTargetName, WarpEndTime).GetLocation();
+
+		// Actor's actual position relative to the warp point
+		const FVector ActorLocationRelative = GetOwner()->GetActorTransform().GetRelativeTransform(WarpPoint->Transform).GetLocation();
+
+		// Distance from the start location to the warp target in the animation
+		const float PerfectDistSq = FVector::DistSquared2D(AnimStartLocation, AnimSyncLocation);
+
+		// Distance from where the actor actually is in game to the warp target
+		const float ActualDistSq = FVector::DistSquared2D(ActorLocationRelative, AnimSyncLocation);
+
+		// Actor is closer than the full authored approach distance. 
+		// Find the time in the animation where the remaining distance to the target matches the actor's actual distance.
+		if (ActualDistSq < PerfectDistSq)
+		{
+			const float SampleInterval = 1.f / 30.f;
+
+			for (float Time = 0.f; Time <= MaxSearchTime; Time += SampleInterval)
+			{
+				const FVector FrameLocation = Asset->GetAlignmentTransform(Binding, InBindings, WarpPoint->Name, Time).GetLocation();
+				const float FrameDistSq = FVector::DistSquared2D(FrameLocation, AnimSyncLocation);
+
+				if (FrameDistSq <= ActualDistSq)
+				{
+					IdealStartTime = Time;
+					break;
+				}
+			}
+		}
+
+		UE_LOG(LogContextualAnim, Verbose, TEXT("IdealStartTime: %.3f (MaxSearch: %.3f) ActualDist: %.1f PerfectDist: %.1f Actor: %s Animation: %s AnimStartLoc: %s AnimSyncLoc: %s ActorLocRelative: %s"),
+			IdealStartTime, MaxSearchTime, FMath::Sqrt(ActualDistSq), FMath::Sqrt(PerfectDistSq),
+			*GetNameSafe(GetOwner()), *GetNameSafe(Binding.GetAnimation()), *AnimStartLocation.ToString(), *AnimSyncLocation.ToString(), *ActorLocationRelative.ToString());
+	}
+
+	return IdealStartTime;
+}
+
 bool UContextualAnimSceneActorComponent::StartContextualAnimScene(const FContextualAnimSceneBindings& InBindings, const TArray<FContextualAnimWarpTarget>& ExternalWarpTargets)
 {
 	UE_LOG(LogContextualAnim, Verbose, TEXT("%-21s UContextualAnimSceneActorComponent::StartContextualAnim Actor: %s"),
@@ -782,7 +850,16 @@ bool UContextualAnimSceneActorComponent::StartContextualAnimScene(const FContext
 			TArray<FContextualAnimWarpPoint> WarpPoints;
 			CalculateWarpPointsForBindings(InBindings, OwnerBinding->GetSectionIdx(), OwnerBinding->GetAnimSetIdx(), WarpPoints);
 
-			JoinScene(InBindings, WarpPoints, ExternalWarpTargets);
+			// Finds ideal start time based on the distance from the actor to the warp target in the animation and the actual distance in game.
+			// Note this time is computed from the initiator's position and all other actors use the same value to stay in sync.
+			float IdealStartTime = 0.f;
+			const UContextualAnimSceneAsset* Asset = InBindings.GetSceneAsset();
+			if (Asset && Asset->ShouldFindIdealAnimStartTime())
+			{
+				IdealStartTime = FindIdealStartTime(InBindings, *OwnerBinding, WarpPoints);
+			}
+
+			JoinScene(InBindings, WarpPoints, ExternalWarpTargets, IdealStartTime);
 
 			for (const FContextualAnimSceneBinding& Binding : InBindings)
 			{
@@ -790,7 +867,7 @@ bool UContextualAnimSceneActorComponent::StartContextualAnimScene(const FContext
 				{
 					if (UContextualAnimSceneActorComponent* Comp = Binding.GetSceneActorComponent())
 					{
-						Comp->JoinScene(InBindings, WarpPoints, ExternalWarpTargets);
+						Comp->JoinScene(InBindings, WarpPoints, ExternalWarpTargets, IdealStartTime);
 					}
 				}
 			}
@@ -798,6 +875,7 @@ bool UContextualAnimSceneActorComponent::StartContextualAnimScene(const FContext
 			RepBindings.Bindings = InBindings;
 			RepBindings.WarpPoints = MoveTemp(WarpPoints);
 			RepBindings.ExternalWarpTargets = ExternalWarpTargets;
+			RepBindings.StartTime = IdealStartTime;
 			RepBindings.IncrementRepCounter();
 
 			MARK_PROPERTY_DIRTY_FROM_NAME(UContextualAnimSceneActorComponent, RepBindings, this);
@@ -942,7 +1020,7 @@ void UContextualAnimSceneActorComponent::OnRep_Bindings()
 		if (ensureAlways(OwnerBinding))
 		{
 			// Join the scene (start playing animation, etc.)
-			JoinScene(RepBindings.Bindings, RepBindings.WarpPoints, RepBindings.ExternalWarpTargets);
+			JoinScene(RepBindings.Bindings, RepBindings.WarpPoints, RepBindings.ExternalWarpTargets, RepBindings.StartTime);
 
 			// RepBindings is only replicated from the initiator of the action.
 			// So now we have to tell everyone else involved in the interaction to join us
@@ -957,7 +1035,7 @@ void UContextualAnimSceneActorComponent::OnRep_Bindings()
 					if (ensureMsgf(Comp, TEXT("%-21s UContextualAnimSceneActorComponent::OnRep_Bindings Actor: %s Current UContextualAnimSceneActorComponent INVALID"),
 						*UEnum::GetValueAsString(TEXT("Engine.ENetRole"), GetOwner()->GetLocalRole()), *GetNameSafe(Binding.GetActor())))
 					{
-						Comp->JoinScene(RepBindings.Bindings, RepBindings.WarpPoints, RepBindings.ExternalWarpTargets);
+						Comp->JoinScene(RepBindings.Bindings, RepBindings.WarpPoints, RepBindings.ExternalWarpTargets, RepBindings.StartTime);
 					}
 				}
 			}
@@ -1106,7 +1184,7 @@ void UContextualAnimSceneActorComponent::OnLeftScene()
 	// This function will be removed
 }
 
-void UContextualAnimSceneActorComponent::JoinScene(const FContextualAnimSceneBindings& InBindings, const TArray<FContextualAnimWarpPoint> WarpPoints, const TArray<FContextualAnimWarpTarget>& ExternalWarpTargets)
+void UContextualAnimSceneActorComponent::JoinScene(const FContextualAnimSceneBindings& InBindings, const TArray<FContextualAnimWarpPoint> WarpPoints, const TArray<FContextualAnimWarpTarget>& ExternalWarpTargets, float StartTime)
 {
 	if (Bindings.IsValid())
 	{
@@ -1123,7 +1201,7 @@ void UContextualAnimSceneActorComponent::JoinScene(const FContextualAnimSceneBin
 		Bindings = InBindings;
 
 		UAnimSequenceBase* AnimationSequence = Binding->GetAnimation();
-		PlayAnimation_Internal(AnimationSequence, 0.f, Bindings.ShouldBindingSyncAnimation(*Binding));
+		PlayAnimation_Internal(AnimationSequence, StartTime, Bindings.ShouldBindingSyncAnimation(*Binding));
 
 		AddOrUpdateWarpTargets(Binding->GetSectionIdx(), Binding->GetAnimSetIdx(), WarpPoints, ExternalWarpTargets);
 
diff --git a/Engine/Plugins/Experimental/Animation/ContextualAnimation/Source/ContextualAnimation/Private/ContextualAnimUtilities.cpp b/Engine/Plugins/Experimental/Animation/ContextualAnimation/Source/ContextualAnimation/Private/ContextualAnimUtilities.cpp
--- a/Engine/Plugins/Experimental/Animation/ContextualAnimation/Source/ContextualAnimation/Private/ContextualAnimUtilities.cpp
+++ b/Engine/Plugins/Experimental/Animation/ContextualAnimation/Source/ContextualAnimation/Private/ContextualAnimUtilities.cpp
@@ -314,15 +314,15 @@ void UContextualAnimUtilities::GetStartAndEndTimeForWarpSection(const UAnimSeque
 	}
 }
 
-void UContextualAnimUtilities::GetStartAndEndTimeForWarpSection(const UAnimSequenceBase* Animation, int32 WarpSectionIndex, float& OutStartTime, float& OutEndTime)
+FName UContextualAnimUtilities::GetStartAndEndTimeForWarpSection(const UAnimSequenceBase* Animation, int32 DesiredWarpSectionOrdinal, float& OutStartTime, float& OutEndTime)
 {
 	OutStartTime = 0.f;
 	OutEndTime = 0.f;
 
-	if (Animation && WarpSectionIndex >= 0)
+	if (Animation && DesiredWarpSectionOrdinal >= 0)
 	{
-		FName LastWarpTargetName = NAME_None;
-		int32 LastWarpSectionIndex = INDEX_NONE;
+		FName CurrentSectionName = NAME_None;
+		int32 CurrentSectionOrdinal = INDEX_NONE;
 
 		for (int32 Idx = 0; Idx < Animation->Notifies.Num(); Idx++)
 		{
@@ -335,42 +335,52 @@ void UContextualAnimUtilities::GetStartAndEndTimeForWarpSection(const UAnimSeque
 					if (WarpTargetName != NAME_None)
 					{
 						// First valid warping window. Initialize everything
-						if (LastWarpSectionIndex == INDEX_NONE)
+						if (CurrentSectionOrdinal == INDEX_NONE)
 						{
-							LastWarpTargetName = WarpTargetName;
+							CurrentSectionName = WarpTargetName;
+							CurrentSectionOrdinal = 0;
 							OutStartTime = NotifyEvent.GetTriggerTime();
 							OutEndTime = NotifyEvent.GetEndTriggerTime();
-							LastWarpSectionIndex = 0;
 						}
-						// If we hit another warping window but the sync point is the same as the previous, update SyncTime.
-						// This is to deal with cases where a first short window is used to face the alignment point and a second one to perform the rest of the warp
-						else if (WarpTargetName == LastWarpTargetName)
+						// Same section (e.g. a rotation window followed by the main warp window).
+						// Keep the earliest start time, extend end time to cover all sub-windows.
+						else if (WarpTargetName == CurrentSectionName)
 						{
-							OutStartTime = NotifyEvent.GetTriggerTime();
-							OutEndTime = NotifyEvent.GetEndTriggerTime();
+							OutEndTime = FMath::Max(OutEndTime, NotifyEvent.GetEndTriggerTime());
 						}
-						// If we hit another warping window but with a different sync point name means that we have hit the first window of another warping section
+						// Different warp target name means we've hit the first window of a new warp section
 						else
 						{
-							// If we haven't reached the desired WarpSection yet. Update control vars and keep moving
-							if (WarpSectionIndex > LastWarpSectionIndex)
-							{
-								LastWarpTargetName = WarpTargetName;
-								OutStartTime = NotifyEvent.GetTriggerTime();
-								OutEndTime = NotifyEvent.GetEndTriggerTime();
-								LastWarpSectionIndex++;
-							}
-							// Otherwise, stop here and return the value of the last window we found
-							else
+							// We already found the desired section and this is a different one, so we're done
+							// Note this assumes notifies are grouped. E.g, we don't have A B A
+							if (CurrentSectionOrdinal == DesiredWarpSectionOrdinal)
 							{
 								break;
 							}
+
+							// Haven't reached the desired section yet. Advance to this new section and keep going.
+							CurrentSectionName = WarpTargetName;
+							CurrentSectionOrdinal++;
+							OutStartTime = NotifyEvent.GetTriggerTime();
+							OutEndTime = NotifyEvent.GetEndTriggerTime();
 						}
 					}
 				}
 			}
 		}
+
+		// Requested ordinal was beyond the number of sections found
+		if (CurrentSectionOrdinal != DesiredWarpSectionOrdinal)
+		{
+			OutStartTime = 0.f;
+			OutEndTime = 0.f;
+			return NAME_None;
+		}
+
+		return CurrentSectionName;
 	}
+
+	return NAME_None;
 }
 
 void UContextualAnimUtilities::GenerateAlignmentTrack(const UContextualAnimSceneAsset& SceneAsset, const FContextualAnimSceneBinding& InBinding, const FContextualAnimSceneBindings& InBindings, FContextualAnimAlignmentTrackContainer& Container)
diff --git a/Engine/Plugins/Experimental/Animation/ContextualAnimation/Source/ContextualAnimation/Public/ContextualAnimSceneActorComponent.h b/Engine/Plugins/Experimental/Animation/ContextualAnimation/Source/ContextualAnimation/Public/ContextualAnimSceneActorComponent.h
--- a/Engine/Plugins/Experimental/Animation/ContextualAnimation/Source/ContextualAnimation/Public/ContextualAnimSceneActorComponent.h
+++ b/Engine/Plugins/Experimental/Animation/ContextualAnimation/Source/ContextualAnimation/Public/ContextualAnimSceneActorComponent.h
@@ -77,12 +77,16 @@ struct FContextualAnimRepBindingsData : public FContextualAnimRepData
 	UPROPERTY()
 	TArray<FContextualAnimWarpTarget> ExternalWarpTargets;
 
+	UPROPERTY()
+	float StartTime = 0.f;
+
 	void Reset()
 	{
 		RepCounter = 0;
 		Bindings.Reset();
 		WarpPoints.Reset();
 		ExternalWarpTargets.Reset();
+		StartTime = 0.f;
 	}
 };
 
@@ -341,7 +345,7 @@ protected:
 	// @TODO: These two functions are going to replace OnJoinedScene and OnLeftScene
 	// main different is that these new functions are taking care of animation playback too
 
-	UE_API void JoinScene(const FContextualAnimSceneBindings& InBindings, const TArray<FContextualAnimWarpPoint> WarpPoints, const TArray<FContextualAnimWarpTarget>& ExternalWarpTargets);
+	UE_API void JoinScene(const FContextualAnimSceneBindings& InBindings, const TArray<FContextualAnimWarpPoint> WarpPoints, const TArray<FContextualAnimWarpTarget>& ExternalWarpTargets, float StartTime = 0.f);
 
 	UE_API virtual void OnJoinScene(const FContextualAnimSceneBinding& Binding);
 
@@ -378,6 +382,13 @@ protected:
 
 	UE_API virtual void OtherActorLeftScene(AActor& Actor);
 
+	/** 
+	 * Finds ideal start time for the interaction based on the distance from the actor to the warp target in the animation and the actual distance in game. 
+	 * Effectively skipping the portion of the 'approach-motion' the actor doesn't need.
+     * Returns 0 if the actor is far enough to play the full animation from the beginning.
+	 */
+	float FindIdealStartTime(const FContextualAnimSceneBindings& InBindings, const FContextualAnimSceneBinding& Binding, const TArray<FContextualAnimWarpPoint>& WarpPoints) const;
+
 private:
 
 	TArray<TWeakObjectPtr<const UAnimMontage>, TInlineAllocator<5>> AnimsPlayed;
diff --git a/Engine/Plugins/Experimental/Animation/ContextualAnimation/Source/ContextualAnimation/Public/ContextualAnimSceneAsset.h b/Engine/Plugins/Experimental/Animation/ContextualAnimation/Source/ContextualAnimation/Public/ContextualAnimSceneAsset.h
--- a/Engine/Plugins/Experimental/Animation/ContextualAnimation/Source/ContextualAnimation/Public/ContextualAnimSceneAsset.h
+++ b/Engine/Plugins/Experimental/Animation/ContextualAnimation/Source/ContextualAnimation/Public/ContextualAnimSceneAsset.h
@@ -306,7 +306,9 @@ public:
 	inline bool ShouldPrecomputeAlignmentTracks() const { return bPrecomputeAlignmentTracks; }
 	inline bool ShouldIgnoreClientMovementErrorChecksAndCorrection() const { return bIgnoreClientMovementErrorChecksAndCorrection; }
 	inline bool ShouldDisableMovementReplicationForSimulatedProxy() const { return bDisableMovementReplicationForSimulatedProxy; }
-		
+	inline bool ShouldFindIdealAnimStartTime() const { return bFindIdealAnimStartTime; }
+	inline float GetFindIdealStartTimeMargin() const { return FindIdealStartTimeMargin; }
+
 	const TArray<FContextualAnimAttachmentParams>& GetAttachmentParams() const { return AttachmentParams; }
 
 	const FContextualAnimAttachmentParams* GetAttachmentParamsForRole(FName Role) const
@@ -442,6 +444,20 @@ protected:
 	UPROPERTY(EditAnywhere, Category = "Settings", meta = (DisplayName = "IK Target Params"))
 	FContextualAnimIKTargetParams IKTargetParams;
 
+	/** 
+	 * Whether to find ideal start time for the interaction based on the distance from the actor to the warp target in the animation and the actual distance in game.
+	 * Effectively skipping the portion of the 'approach-motion' the actor doesn't need. 
+	 */
+	UPROPERTY(EditAnywhere, Category = "Settings")
+	bool bFindIdealAnimStartTime = false;
+
+	/**
+	 * Time (in seconds) reserved at the end of the warp window to ensure motion warping has enough room to reach the target.
+	 * The ideal start time search will not go beyond (WarpEndTime - FindIdealStartTimeMargin).
+	 */ 
+	UPROPERTY(EditAnywhere, Category = "Settings", meta = (EditCondition = "bFindIdealAnimStartTime", EditConditionHides))
+	float FindIdealStartTimeMargin = 0.2f;
+
 	/** Whether to ignore movement error checks and corrections during the interaction for player characters */
 	UPROPERTY(EditAnywhere, Category = "Settings", AdvancedDisplay)
 	bool bIgnoreClientMovementErrorChecksAndCorrection = true;
diff --git a/Engine/Plugins/Experimental/Animation/ContextualAnimation/Source/ContextualAnimation/Public/ContextualAnimUtilities.h b/Engine/Plugins/Experimental/Animation/ContextualAnimation/Source/ContextualAnimation/Public/ContextualAnimUtilities.h
--- a/Engine/Plugins/Experimental/Animation/ContextualAnimation/Source/ContextualAnimation/Public/ContextualAnimUtilities.h
+++ b/Engine/Plugins/Experimental/Animation/ContextualAnimation/Source/ContextualAnimation/Public/ContextualAnimUtilities.h
@@ -73,7 +73,9 @@ public:
 	static UE_API float GetSyncTimeForWarpSection(const UAnimSequenceBase* Animation, const FName& WarpSectionName);
 	
 	static UE_API void GetStartAndEndTimeForWarpSection(const UAnimSequenceBase* Animation, const FName& WarpSectionName, float& OutStartTime, float& OutEndTime);
-	static UE_API void GetStartAndEndTimeForWarpSection(const UAnimSequenceBase* Animation, int32 WarpSectionIndex, float& OutStartTime, float& OutEndTime);
+
+	/** Returns the start/end time and warp target name of the Nth distinct warp section found sequentially in the animation's motion warping notifies */
+	static UE_API FName GetStartAndEndTimeForWarpSection(const UAnimSequenceBase* Animation, int32 DesiredWarpSectionOrdinal, float& OutStartTime, float& OutEndTime);
 	
 	static UE_API void GenerateAlignmentTrack(const UContextualAnimSceneAsset& SceneAsset, const FContextualAnimSceneBinding& InBinding, const FContextualAnimSceneBindings& InBindings, FContextualAnimAlignmentTrackContainer& Container);
 	static UE_API void GenerateIKTargetTrack(const UContextualAnimSceneAsset& SceneAsset, const FContextualAnimSceneBinding& InBinding, const FContextualAnimSceneBindings& InBindings, FContextualAnimAlignmentTrackContainer& Container);
