Contextual Anim Editor - Added warning message to validate that a preview actor class has a ContextualAnimSceneActorComponent
#Synced-CL 51627494

diff --git a/Engine/Plugins/Experimental/Animation/ContextualAnimation/Source/ContextualAnimationEditor/Private/ContextualAnimEdMode.cpp b/Engine/Plugins/Experimental/Animation/ContextualAnimation/Source/ContextualAnimationEditor/Private/ContextualAnimEdMode.cpp
--- a/Engine/Plugins/Experimental/Animation/ContextualAnimation/Source/ContextualAnimationEditor/Private/ContextualAnimEdMode.cpp
+++ b/Engine/Plugins/Experimental/Animation/ContextualAnimation/Source/ContextualAnimationEditor/Private/ContextualAnimEdMode.cpp
@@ -50,9 +50,21 @@ void FContextualAnimEdMode::Render(const FSceneView* View, FViewport* Viewport,
 		if (const UContextualAnimSceneAsset* SceneAsset = ViewModel->GetSceneAsset())
 		{
 			const FContextualAnimSceneBindings& Bindings = ViewModel->SceneBindings;
+
+			// If any bound actor became invalid (e.g., Blueprint recompilation destroying the old instance),
+			// refresh the preview scene to re-spawn actors from the current classes.
+			for (const FContextualAnimSceneBinding& Binding : Bindings)
+			{
+				if (!Binding.GetActor())
+				{
+					ViewModel->SetDefaultMode();
+					return;
+				}
+			}
+
 			const FName PrimaryRole = SceneAsset->GetPrimaryRole();
 			const FContextualAnimSceneBinding* PrimaryBinding = Bindings.FindBindingByRole(PrimaryRole);
-			
+
 			if(PrimaryBinding == nullptr)
 			{
 				return;
diff --git a/Engine/Plugins/Experimental/Animation/ContextualAnimation/Source/ContextualAnimationEditor/Private/ContextualAnimViewModel.cpp b/Engine/Plugins/Experimental/Animation/ContextualAnimation/Source/ContextualAnimationEditor/Private/ContextualAnimViewModel.cpp
--- a/Engine/Plugins/Experimental/Animation/ContextualAnimation/Source/ContextualAnimationEditor/Private/ContextualAnimViewModel.cpp
+++ b/Engine/Plugins/Experimental/Animation/ContextualAnimation/Source/ContextualAnimationEditor/Private/ContextualAnimViewModel.cpp
@@ -30,7 +30,11 @@
 #include "AnimPreviewInstance.h"
 #include "ContextualAnimSelectionCriterion.h"
 #include "ScopedTransaction.h"
+#include "Misc/MessageDialog.h"
 #include "UObject/Package.h"
+#include "Engine/SCS_Node.h"
+#include "Engine/SimpleConstructionScript.h"
+#include "Engine/BlueprintGeneratedClass.h"
 #include "Engine/StaticMesh.h"
 #include "Components/StaticMeshComponent.h"
 #include "DetailCustomizations/ContextualAnimSceneAssetDetailCustom.h"
@@ -39,6 +43,18 @@
 
 static const FName PreviewActorTag = FName(TEXT("__PreviewActor__"));
 
+static void ShowMissingSceneActorComponentWarning(const FString& ClassName)
+{
+	const FText WarningMsg = FText::Format(
+		LOCTEXT("MissingSceneActorComponent",
+			"The PreviewActorClass '{0}' does not contain a UContextualAnimSceneActorComponent. "
+			"This component is required for the Contextual Animation editor to function correctly. "
+			"The preview actor selection has been reverted."),
+		FText::FromString(ClassName));
+
+	FMessageDialog::Open(EAppMsgType::Ok, WarningMsg);
+}
+
 FContextualAnimViewModel::FContextualAnimViewModel()
 	: SceneAsset(nullptr)
 	, MovieSceneSequence(nullptr)
@@ -314,6 +330,27 @@ AActor* FContextualAnimViewModel::SpawnPreviewActor(const FContextualAnimTrack&
 				{
 					IContextualAnimActorInterface::Execute_RefreshEditorPreview(PreviewActor);
 				}
+
+				// Validate that the spawned actor has the required UContextualAnimSceneActorComponent.
+				// The user may have removed the component from the Blueprint after initially selecting a valid class.
+				if (!PreviewActor->FindComponentByClass<UContextualAnimSceneActorComponent>())
+				{
+					ShowMissingSceneActorComponentWarning(PreviewActorClass->GetName());
+
+					PreviewActor->Destroy();
+					PreviewActor = nullptr;
+
+					// Reset the preview data to prevent repeated warnings on subsequent refreshes
+					for (FContextualAnimActorPreviewData& Data : GetSceneAsset()->OverridePreviewData)
+					{
+						if (Data.Role == AnimTrack.Role && Data.Type == EContextualAnimActorPreviewType::Actor)
+						{
+							Data.PreviewActorClass.Reset();
+							break;
+						}
+					}
+					GetSceneAsset()->MarkPackageDirty();
+				}
 			}
 			else
 			{
@@ -358,7 +395,10 @@ AActor* FContextualAnimViewModel::SpawnPreviewActor(const FContextualAnimTrack&
 			}
 		}
 
-		PreviewActor->Tags.Add(PreviewActorTag);
+		if (PreviewActor)
+		{
+			PreviewActor->Tags.Add(PreviewActorTag);
+		}
 	}
 
 	return PreviewActor;
@@ -1392,10 +1432,80 @@ void FContextualAnimViewModel::OnFinishedChangingProperties(const FPropertyChang
 
 	UE_LOG(LogContextualAnim, Verbose, TEXT("FContextualAnimViewModel::OnFinishedChangingProperties MemberPropertyName: %s PropertyName: %s"), *MemberPropertyName.ToString(), *PropertyName.ToString());
 
+	// Validate PreviewActorClass has the required UContextualAnimSceneActorComponent
+	if (MemberPropertyName == GET_MEMBER_NAME_CHECKED(UContextualAnimSceneAsset, OverridePreviewData) &&
+		PropertyName == GET_MEMBER_NAME_CHECKED(FContextualAnimActorPreviewData, PreviewActorClass))
+	{
+		bool bReverted = false;
+		for (FContextualAnimActorPreviewData& PreviewData : SceneAsset->OverridePreviewData)
+		{
+			if (PreviewData.Type == EContextualAnimActorPreviewType::Actor)
+			{
+				if (UClass* ActorClass = PreviewData.PreviewActorClass.LoadSynchronous())
+				{
+					// Check native components on the CDO
+					bool bHasSceneActorComponent = false;
+					if (AActor* CDO = ActorClass->GetDefaultObject<AActor>())
+					{
+						bHasSceneActorComponent = CDO->FindComponentByClass<UContextualAnimSceneActorComponent>() != nullptr;
+					}
+
+					// Check Blueprint SCS for components added in the Blueprint editor
+					if (!bHasSceneActorComponent)
+					{
+						for (UClass* CurrentClass = ActorClass; CurrentClass && !bHasSceneActorComponent; CurrentClass = CurrentClass->GetSuperClass())
+						{
+							if (UBlueprintGeneratedClass* BPGC = Cast<UBlueprintGeneratedClass>(CurrentClass))
+							{
+								if (USimpleConstructionScript* SCS = BPGC->SimpleConstructionScript)
+								{
+									for (USCS_Node* Node : SCS->GetAllNodes())
+									{
+										if (Node && Node->ComponentTemplate && Node->ComponentTemplate->IsA<UContextualAnimSceneActorComponent>())
+										{
+											bHasSceneActorComponent = true;
+											break;
+										}
+									}
+								}
+							}
+						}
+					}
+
+					if (!bHasSceneActorComponent)
+					{
+						ShowMissingSceneActorComponentWarning(ActorClass->GetName());
+
+						PreviewData.PreviewActorClass.Reset();
+						bReverted = true;
+					}
+				}
+			}
+		}
+
+		if (bReverted)
+		{
+			SceneAsset->MarkPackageDirty();
+			RefreshDetailsView();
+		}
+		else
+		{
+			SetDefaultMode();
+		}
+		return;
+	}
+
+	// Refresh preview scene when OverridePreviewData changes
+	if (MemberPropertyName == GET_MEMBER_NAME_CHECKED(UContextualAnimSceneAsset, OverridePreviewData))
+	{
+		SetDefaultMode();
+		return;
+	}
+
 	// Refresh preview scene if necessary
 	if (MemberPropertyName == GET_MEMBER_NAME_CHECKED(UContextualAnimSceneAsset, Sections) &&
 		(PropertyName == GET_MEMBER_NAME_CHECKED(UContextualAnimSceneAsset, Sections) ||
-		 PropertyName == GET_MEMBER_NAME_CHECKED(FContextualAnimSceneSection, AnimSets) || 
+		 PropertyName == GET_MEMBER_NAME_CHECKED(FContextualAnimSceneSection, AnimSets) ||
 		 PropertyName == GET_MEMBER_NAME_CHECKED(FContextualAnimTrack, Animation) ||
 		 PropertyName == GET_MEMBER_NAME_CHECKED(FContextualAnimTrack, MeshToScene)))
 	{
