Contextual Anim Editor:
- Added multi-section preview.
- Improved visibility and format of info text showed on the viewport and added idx and name of the active section.
- Selected Timeline Section is now highligthed in both single and multi-section modes.
- Moved preview related properties to a new 'Preview' category on the Details panel.
- Fixed issue where Tracks would not auto expand automatically when timeline is refreshed.

#Synced-CL 52203160

diff --git a/Engine/Plugins/Experimental/Animation/ContextualAnimation/Source/ContextualAnimation/Public/ContextualAnimSceneAsset.h b/Engine/Plugins/Experimental/Animation/ContextualAnimation/Source/ContextualAnimation/Public/ContextualAnimSceneAsset.h
--- a/Engine/Plugins/Experimental/Animation/ContextualAnimation/Source/ContextualAnimation/Public/ContextualAnimSceneAsset.h
+++ b/Engine/Plugins/Experimental/Animation/ContextualAnimation/Source/ContextualAnimation/Public/ContextualAnimSceneAsset.h
@@ -242,6 +242,8 @@ enum class EContextualAnimCriterionToConsider : uint8
 	Other
 };
 
+#if WITH_EDITORONLY_DATA
+
 UENUM(BlueprintType)
 enum class EContextualAnimActorPreviewType : uint8
 {
@@ -275,6 +277,27 @@ struct FContextualAnimActorPreviewData
 	TSoftClassPtr<class AActor> PreviewActorClass;
 };
 
+/** Entry in a preview sequence that defines which section to play and how many times */
+USTRUCT(BlueprintType)
+struct FContextualAnimPreviewSequenceEntry
+{
+	GENERATED_BODY()
+
+	/** Which section to play */
+	UPROPERTY(EditAnywhere, Category = "Defaults", meta = (GetOptions = "GetSectionNames"))
+	FName SectionName;
+
+	/** How many times to play this section (1 = play once, 2 = repeat twice, etc.) */
+	UPROPERTY(EditAnywhere, Category = "Defaults", meta = (ClampMin = "1", ClampMax = "100"))
+	int32 LoopCount = 1;
+
+	/** Which anim set to use. INDEX_NONE = use the editor's active anim set for that section */
+	UPROPERTY(EditAnywhere, Category = "Defaults")
+	int32 AnimSetIdx = INDEX_NONE;
+};
+
+#endif  // WITH_EDITORONLY_DATA
+
 UCLASS(MinimalAPI, Blueprintable)
 class UContextualAnimSceneAsset : public UDataAsset
 {
@@ -331,6 +354,7 @@ public:
 
 	UE_API const FTransform& GetMeshToComponentForRole(const FName& Role) const;
 
+	UFUNCTION()
 	UE_API TArray<FName> GetSectionNames() const;
 
 	UE_API int32 GetNumSections() const;
@@ -413,12 +437,16 @@ protected:
 
 #if WITH_EDITORONLY_DATA
 
-	UPROPERTY(EditAnywhere, Category = "Settings", meta = (TitleProperty = "Role"))
+	UPROPERTY(EditAnywhere, Category = "Preview", meta = (TitleProperty = "Role"))
 	TArray<FContextualAnimActorPreviewData> OverridePreviewData;
 
-	UPROPERTY(EditAnywhere, Category = "Settings", meta = (TitleProperty = "Role", MustImplement = "/Script/ContextualAnimation.ContextualAnimOverrideInterface"))
+	UPROPERTY(EditAnywhere, Category = "Preview", meta = (TitleProperty = "Role", MustImplement = "/Script/ContextualAnimation.ContextualAnimOverrideInterface"))
 	TSubclassOf<UObject> PreviewOverrideProvider = nullptr;;
 
+	/** Defines a sequence of sections for multi-section preview (e.g. Intro -> Loop x3 -> Outro) */
+	UPROPERTY(EditAnywhere, Category = "Preview", meta = (TitleProperty = "SectionName"))
+	TArray<FContextualAnimPreviewSequenceEntry> MultiSectionPreviewSequence;
+
 #endif // WITH_EDITORONLY_DATA
 
 	UPROPERTY(EditAnywhere, Category = "Defaults")
diff --git a/Engine/Plugins/Experimental/Animation/ContextualAnimation/Source/ContextualAnimationEditor/Private/ContextualAnimEdMode.cpp b/Engine/Plugins/Experimental/Animation/ContextualAnimation/Source/ContextualAnimationEditor/Private/ContextualAnimEdMode.cpp
--- a/Engine/Plugins/Experimental/Animation/ContextualAnimation/Source/ContextualAnimationEditor/Private/ContextualAnimEdMode.cpp
+++ b/Engine/Plugins/Experimental/Animation/ContextualAnimation/Source/ContextualAnimationEditor/Private/ContextualAnimEdMode.cpp
@@ -32,6 +32,11 @@ FContextualAnimEdMode::~FContextualAnimEdMode()
 void FContextualAnimEdMode::Tick(FEditorViewportClient* ViewportClient, float DeltaTime)
 {
 	FEdMode::Tick(ViewportClient, DeltaTime);
+
+	if (ViewModel)
+	{
+		ViewModel->TickSimulateMode(DeltaTime);
+	}
 }
 
 void FContextualAnimEdMode::Render(const FSceneView* View, FViewport* Viewport, FPrimitiveDrawInterface* PDI)
@@ -250,7 +255,25 @@ void FContextualAnimEdMode::DrawHUD(FEditorViewportClient* ViewportClient, FView
 
 	if (ViewModel)
 	{
-		FCanvasTextItem TextItem(FVector2D(10.f, 40.f), ViewModel->GetSelectionDebugText(), GEngine->GetSmallFont(), FLinearColor::White);
+		const FText DebugText = ViewModel->GetSelectionDebugText();
+		const UFont* Font = GEngine->GetSmallFont();
+		const FVector2D TextPos(10.f, 10.f);
+		const float Padding = 6.f;
+
+		// Measure text size for the background
+		int32 TextWidth = 0;
+		int32 TextHeight = 0;
+		StringSize(Font, TextWidth, TextHeight, *DebugText.ToString());
+
+		// Draw semi-transparent black background
+		const FVector2D BackgroundPos(TextPos.X - Padding, TextPos.Y - Padding);
+		const FVector2D BackgroundSize(TextWidth + Padding * 2.f, TextHeight + Padding * 2.f);
+		FCanvasTileItem BackgroundTile(BackgroundPos, BackgroundSize, FLinearColor(0.f, 0.f, 0.f, 0.5f));
+		BackgroundTile.BlendMode = SE_BLEND_Translucent;
+		Canvas->DrawItem(BackgroundTile);
+
+		// Draw text
+		FCanvasTextItem TextItem(TextPos, DebugText, Font, FLinearColor::White);
 		Canvas->DrawItem(TextItem);
 	}
 }
diff --git a/Engine/Plugins/Experimental/Animation/ContextualAnimation/Source/ContextualAnimationEditor/Private/ContextualAnimMovieSceneSection.cpp b/Engine/Plugins/Experimental/Animation/ContextualAnimation/Source/ContextualAnimationEditor/Private/ContextualAnimMovieSceneSection.cpp
--- a/Engine/Plugins/Experimental/Animation/ContextualAnimation/Source/ContextualAnimationEditor/Private/ContextualAnimMovieSceneSection.cpp
+++ b/Engine/Plugins/Experimental/Animation/ContextualAnimation/Source/ContextualAnimationEditor/Private/ContextualAnimMovieSceneSection.cpp
@@ -15,11 +15,12 @@ UContextualAnimMovieSceneTrack& UContextualAnimMovieSceneSection::GetOwnerTrack(
 	return *OwnerTrackPtr;
 }
 
-void UContextualAnimMovieSceneSection::Initialize(int32 InSectionIdx, int32 InAnimSetIdx, int32 InAnimTrackIdx)
+void UContextualAnimMovieSceneSection::Initialize(int32 InSectionIdx, int32 InAnimSetIdx, int32 InAnimTrackIdx, int32 InSegmentIdx)
 {
 	SectionIdx = InSectionIdx;
 	AnimSetIdx = InAnimSetIdx;
 	AnimTrackIdx = InAnimTrackIdx;
+	SegmentIdx = InSegmentIdx;
 }
 
 FContextualAnimTrack& UContextualAnimMovieSceneSection::GetAnimTrack() const
diff --git a/Engine/Plugins/Experimental/Animation/ContextualAnimation/Source/ContextualAnimationEditor/Private/ContextualAnimMovieSceneTrackEditor.cpp b/Engine/Plugins/Experimental/Animation/ContextualAnimation/Source/ContextualAnimationEditor/Private/ContextualAnimMovieSceneTrackEditor.cpp
--- a/Engine/Plugins/Experimental/Animation/ContextualAnimation/Source/ContextualAnimationEditor/Private/ContextualAnimMovieSceneTrackEditor.cpp
+++ b/Engine/Plugins/Experimental/Animation/ContextualAnimation/Source/ContextualAnimationEditor/Private/ContextualAnimMovieSceneTrackEditor.cpp
@@ -4,6 +4,9 @@
 #include "ContextualAnimMovieSceneTrack.h"
 #include "ContextualAnimMovieSceneSection.h"
 #include "ContextualAnimViewModel.h"
+#include "SequencerSectionPainter.h"
+#include "Styling/AppStyle.h"
+#include "Styling/StyleColors.h"
 
 #define LOCTEXT_NAMESPACE "FContextualAnimMovieSceneTrackEditor"
 
@@ -57,6 +60,38 @@ bool FContextualAnimMovieSceneTrackEditor::SupportsSequence(UMovieSceneSequence*
 // FContextualAnimSection
 ////////////////////////////////////////////////////////////////////////////////////////////////
 
+int32 FContextualAnimSection::OnPaintSection(FSequencerSectionPainter& InPainter) const
+{
+	int32 LayerId = FSequencerSection::OnPaintSection(InPainter);
+
+	if (UContextualAnimMovieSceneSection* Section = Cast<UContextualAnimMovieSceneSection>(WeakSection.Get()))
+	{
+		const FContextualAnimViewModel& ViewModel = Section->GetOwnerTrack().GetViewModel();
+
+		if (const FContextualAnimTrack* SelectedTrack = ViewModel.GetSelectedAnimTrack())
+		{
+			if (Section->GetSectionIdx() == SelectedTrack->SectionIdx &&
+				Section->GetAnimSetIdx() == SelectedTrack->AnimSetIdx &&
+				Section->GetAnimTrackIdx() == SelectedTrack->AnimTrackIdx)
+			{
+				const ESlateDrawEffect DrawEffects = InPainter.bParentEnabled ? ESlateDrawEffect::None : ESlateDrawEffect::DisabledEffect;
+
+				FSlateDrawElement::MakeBox(
+					InPainter.DrawElements,
+					LayerId,
+					InPainter.SectionGeometry.ToPaintGeometry(),
+					FAppStyle::GetBrush(TEXT("Sequencer.Section.LockedBorder")),
+					DrawEffects,
+					FStyleColors::AccentGreen.GetColor(FWidgetStyle())
+				);
+				++LayerId;
+			}
+		}
+	}
+
+	return LayerId;
+}
+
 FText FContextualAnimSection::GetSectionTitle() const
 {
 	if (UContextualAnimMovieSceneSection* Section = Cast<UContextualAnimMovieSceneSection>(WeakSection.Get()))
diff --git a/Engine/Plugins/Experimental/Animation/ContextualAnimation/Source/ContextualAnimationEditor/Private/ContextualAnimViewModel.cpp b/Engine/Plugins/Experimental/Animation/ContextualAnimation/Source/ContextualAnimationEditor/Private/ContextualAnimViewModel.cpp
--- a/Engine/Plugins/Experimental/Animation/ContextualAnimation/Source/ContextualAnimationEditor/Private/ContextualAnimViewModel.cpp
+++ b/Engine/Plugins/Experimental/Animation/ContextualAnimation/Source/ContextualAnimationEditor/Private/ContextualAnimViewModel.cpp
@@ -12,6 +12,7 @@
 #include "ContextualAnimActorInterface.h"
 #include "ContextualAnimPreviewScene.h"
 #include "Animation/AnimMontage.h"
+#include "Animation/AnimSequence.h"
 #include "Editor.h"
 #include "EngineUtils.h"
 #include "Framework/Application/SlateApplication.h"
@@ -55,6 +56,26 @@ static void ShowMissingSceneActorComponentWarning(const FString& ClassName)
 	FMessageDialog::Open(EAppMsgType::Ok, WarningMsg);
 }
 
+/** Returns true if the underlying animation sequence at the given time has bForceRootLock enabled. */
+static bool HasForceRootLock(const UAnimSequenceBase* Animation, float Time)
+{
+	if (const UAnimMontage* Montage = Cast<UAnimMontage>(Animation))
+	{
+		if (const FAnimSegment* Segment = Montage->SlotAnimTracks[0].AnimTrack.GetSegmentAtTime(Time))
+		{
+			if (const UAnimSequence* AnimSequence = Cast<UAnimSequence>(Segment->GetAnimReference()))
+			{
+				return AnimSequence->bForceRootLock;
+			}
+		}
+	}
+	else if (const UAnimSequence* AnimSequence = Cast<UAnimSequence>(Animation))
+	{
+		return AnimSequence->bForceRootLock;
+	}
+	return false;
+}
+
 FContextualAnimViewModel::FContextualAnimViewModel()
 	: SceneAsset(nullptr)
 	, MovieSceneSequence(nullptr)
@@ -201,6 +222,16 @@ int32 FContextualAnimViewModel::GetActiveAnimSetForSection(int32 SectionIdx) con
 	return ActiveSetIdxPtr ? *ActiveSetIdxPtr : INDEX_NONE;
 }
 
+int32 FContextualAnimViewModel::GetDetailsSectionIdx() const
+{
+	return SelectionInfo.SectionIdx != INDEX_NONE ? SelectionInfo.SectionIdx : GetActiveSection();
+}
+
+int32 FContextualAnimViewModel::GetDetailsAnimSetIdx() const
+{
+	return SelectionInfo.AnimSetIdx != INDEX_NONE ? SelectionInfo.AnimSetIdx : GetActiveAnimSetForSection(GetDetailsSectionIdx());
+}
+
 AActor* FContextualAnimViewModel::SpawnPreviewActor(const FContextualAnimTrack& AnimTrack)
 {
 	AActor* PreviewActor = nullptr;
@@ -429,101 +460,133 @@ void FContextualAnimViewModel::SetDefaultMode()
 
 	ResetTimeline();
 
-	if (!SceneAsset->Sections.IsValidIndex(ActiveSectionIdx))
+	// Create one MovieScene track per role
+	auto CreateMovieSceneTracksForRoles = [this]()
 	{
-		RefreshPreviewScene();
-		return;
-	}
-
-	FContextualAnimSceneSection& ContextualAnimSection = SceneAsset->Sections[ActiveSectionIdx];
-
-	// If the roles asset was edited, our contextual anim asset may have tracks for roles that don't exist in the roles asset.
-	// We allow this, log a warning, and display this in the UI for the user.
-	TArray<FName> ValidRoles = SceneAsset->GetRoles();
-	TArray<FName> ActualRoles;
-	for (int32 AnimSetIdx = 0; AnimSetIdx < ContextualAnimSection.AnimSets.Num(); AnimSetIdx++)
-	{
-		FContextualAnimSet& AnimSet = ContextualAnimSection.AnimSets[AnimSetIdx];
-		for (int32 AnimTrackIdx = 0; AnimTrackIdx < AnimSet.Tracks.Num(); AnimTrackIdx++)
+		const TArray<FName> Roles = SceneAsset->GetRoles();
+		for (FName Role : Roles)
 		{
-			const FContextualAnimTrack& AnimTrack = AnimSet.Tracks[AnimTrackIdx];
-			ActualRoles.AddUnique(AnimTrack.Role);
+			UContextualAnimMovieSceneTrack* MovieSceneAnimTrack = MovieSceneSequence->GetMovieScene()->AddTrack<UContextualAnimMovieSceneTrack>();
+			check(MovieSceneAnimTrack);
+			MovieSceneAnimTrack->Initialize(Role);
 		}
-	}
+	};
 
-	for (FName Role : ActualRoles)
-	{
-		UContextualAnimMovieSceneTrack* MovieSceneAnimTrack = MovieSceneSequence->GetMovieScene()->AddTrack<UContextualAnimMovieSceneTrack>();
-		check(MovieSceneAnimTrack);
-
-		MovieSceneAnimTrack->Initialize(Role);
-	}
+	// Create MovieScene sections for all anim sets in a given section
+	TArray<FName, TInlineAllocator<5>> MissingRoles;
 
-	float EmptyAnimSectionLength = 0.f;
-	for (int32 AnimSetIdx = 0; AnimSetIdx < ContextualAnimSection.AnimSets.Num(); AnimSetIdx++)
+	auto CreateMovieSceneSectionsForAnimSets = [this, &MissingRoles](const FContextualAnimSceneSection& Section, int32 SectionIdx, float StartTime, float FallbackDuration, int32 ActiveAnimSetIdx, int32 SegmentIdx = INDEX_NONE)
 	{
-		const FContextualAnimSet& AnimSet = ContextualAnimSection.AnimSets[AnimSetIdx];
-		for (int32 AnimTrackIdx = 0; AnimTrackIdx < AnimSet.Tracks.Num(); AnimTrackIdx++)
-		{
-			const FContextualAnimTrack& AnimTrack = AnimSet.Tracks[AnimTrackIdx];
-			if(AnimTrack.Animation)
-			{
-				EmptyAnimSectionLength = FMath::Max(EmptyAnimSectionLength, AnimTrack.Animation->GetPlayLength());
-			}
-		}
-	}
-
-	EmptyAnimSectionLength = FMath::Max(1.f, EmptyAnimSectionLength);
+		const FFrameRate TickResolution = MovieSceneSequence->GetMovieScene()->GetTickResolution();
 
-	for (int32 AnimSetIdx = 0; AnimSetIdx < ContextualAnimSection.AnimSets.Num(); AnimSetIdx++)
-	{
-		FContextualAnimSet& AnimSet = ContextualAnimSection.AnimSets[AnimSetIdx];
-		for (int32 AnimTrackIdx = 0; AnimTrackIdx < AnimSet.Tracks.Num(); AnimTrackIdx++)
+		for (int32 AnimSetIdx = 0; AnimSetIdx < Section.GetNumAnimSets(); AnimSetIdx++)
 		{
-			const FContextualAnimTrack& AnimTrack = AnimSet.Tracks[AnimTrackIdx];
+			const FContextualAnimSet* AnimSet = Section.GetAnimSet(AnimSetIdx);
+			check(AnimSet);
 
-			UContextualAnimMovieSceneTrack* MovieSceneTrack = FindTrackByRole(AnimTrack.Role);
-			check(MovieSceneTrack)
-
-			UContextualAnimMovieSceneSection* NewSection = nullptr;
-			UAnimSequenceBase* Animation = AnimTrack.Animation;
-			if (Animation)
+			for (int32 AnimTrackIdx = 0; AnimTrackIdx < AnimSet->Tracks.Num(); AnimTrackIdx++)
 			{
-				NewSection = NewObject<UContextualAnimMovieSceneSection>(MovieSceneTrack, UContextualAnimMovieSceneSection::StaticClass(), NAME_None, RF_Transactional);
+				const FContextualAnimTrack& AnimTrack = AnimSet->Tracks[AnimTrackIdx];
+
+				// If the roles asset was edited, we may have tracks for roles that no longer exist in the roles asset. 
+				// Collect those and warn the user
+				UContextualAnimMovieSceneTrack* MovieSceneTrack = FindTrackByRole(AnimTrack.Role);
+				if (MovieSceneTrack == nullptr)
+				{
+					MissingRoles.AddUnique(AnimTrack.Role);
+					continue;
+				}
+
+				UContextualAnimMovieSceneSection* NewSection = NewObject<UContextualAnimMovieSceneSection>(MovieSceneTrack, UContextualAnimMovieSceneSection::StaticClass(), NAME_None, RF_Transactional);
 				check(NewSection);
 
-				const float AnimLength = Animation->GetPlayLength();
-				const FFrameRate TickResolution = MovieSceneSequence->GetMovieScene()->GetTickResolution();
-				NewSection->SetRange(TRange<FFrameNumber>::Inclusive(0, (AnimLength * TickResolution).RoundToFrame()));
+				if (AnimTrack.Animation)
+				{
+					const float AnimLength = FMath::Min(AnimTrack.Animation->GetPlayLength(), FallbackDuration);
+					const FFrameNumber StartFrame = (StartTime * TickResolution).RoundToFrame();
+					const FFrameNumber EndFrame = ((StartTime + AnimLength) * TickResolution).RoundToFrame();
+					NewSection->SetRange(TRange<FFrameNumber>::Inclusive(StartFrame, EndFrame));
+					NewSection->SetRowIndex(AnimSetIdx);
+					NewSection->SetIsActive(AnimSetIdx == ActiveAnimSetIdx);
+				}
+				else
+				{
+					const FFrameNumber StartFrame = (StartTime * TickResolution).RoundToFrame();
+					const FFrameNumber EndFrame = ((StartTime + FallbackDuration) * TickResolution).RoundToFrame();
+					NewSection->SetRange(TRange<FFrameNumber>::Inclusive(StartFrame, EndFrame));
+					NewSection->SetRowIndex(AnimSetIdx);
+					NewSection->SetIsActive(false);
+				}
 
-				NewSection->SetRowIndex(AnimSetIdx);
+				NewSection->Initialize(SectionIdx, AnimSetIdx, AnimTrackIdx, SegmentIdx);
 
-				const int32& ActiveSetIdx = ActiveAnimSetMap.FindOrAdd(ActiveSectionIdx);
-				NewSection->SetIsActive(AnimSetIdx == ActiveSetIdx);
+				MovieSceneTrack->AddSection(*NewSection);
+				MovieSceneTrack->SetTrackRowDisplayName(FText::FromString(FString::Printf(TEXT("%d"), AnimSetIdx)), AnimSetIdx);
 			}
-			else
-			{
-				NewSection = NewObject<UContextualAnimMovieSceneSection>(MovieSceneTrack, UContextualAnimMovieSceneSection::StaticClass(), NAME_None, RF_Transactional);
-				check(NewSection);
+		}
+	};
 
-				const FFrameRate TickResolution = MovieSceneSequence->GetMovieScene()->GetTickResolution();
-				NewSection->SetRange(TRange<FFrameNumber>::Inclusive(0, (EmptyAnimSectionLength * TickResolution).RoundToFrame()));
-				NewSection->SetRowIndex(AnimSetIdx);
-				NewSection->SetIsActive(false);
-			}
+	CreateMovieSceneTracksForRoles();
+
+	// Multi-section: flatten all segments from the PreviewSequence into a single timeline.
+	// Each segment maps to a section/anim set pair and is laid out sequentially.
+	const bool bHasMultiSectionTimeline = TryBuildMultiSectionPreviewTimelineData();
+	if (bHasMultiSectionTimeline)
+	{
+		ActiveSectionIdx = ComputedTimeline[0].SectionIdx;
 
-			NewSection->Initialize(ActiveSectionIdx, AnimSetIdx, AnimTrackIdx);
+		for (int32 SegIdx = 0; SegIdx < ComputedTimeline.Num(); SegIdx++)
+		{
+			const FPreviewTimelineSegment& Segment = ComputedTimeline[SegIdx];
+			const FContextualAnimSceneSection& Section = SceneAsset->Sections[Segment.SectionIdx];
+			CreateMovieSceneSectionsForAnimSets(Section, Segment.SectionIdx, Segment.GlobalStartTime, Segment.Duration, Segment.AnimSetIdx, SegIdx);
+		}
+	}
+	// Single-section: display only the active section starting at time 0.
+	else if (SceneAsset->Sections.IsValidIndex(ActiveSectionIdx))
+	{
+		const FContextualAnimSceneSection& ContextualAnimSection = SceneAsset->Sections[ActiveSectionIdx];
 
-			MovieSceneTrack->AddSection(*NewSection);
-			MovieSceneTrack->SetTrackRowDisplayName(FText::FromString(FString::Printf(TEXT("%d"), AnimSetIdx)), AnimSetIdx);
+		// Calculate duration for empty tracks.
+		float FallbackDuration = 0.f;
+		for (int32 AnimSetIdx = 0; AnimSetIdx < ContextualAnimSection.GetNumAnimSets(); AnimSetIdx++)
+		{
+			if (const FContextualAnimSet* AnimSet = ContextualAnimSection.GetAnimSet(AnimSetIdx))
+			{
+				for (const FContextualAnimTrack& AnimTrack : AnimSet->Tracks)
+				{
+					if (AnimTrack.Animation)
+					{
+						FallbackDuration = FMath::Max(FallbackDuration, AnimTrack.Animation->GetPlayLength());
+					}
+				}
+			}
+		}
+		FallbackDuration = FMath::Max(1.f, FallbackDuration);
 
-			const bool bIsValidRole = ValidRoles.Contains(AnimTrack.Role);
-			if (bIsValidRole == false)
+		const int32 ActiveAnimSetIdx = ActiveAnimSetMap.FindOrAdd(ActiveSectionIdx);
+		CreateMovieSceneSectionsForAnimSets(ContextualAnimSection, ActiveSectionIdx, 0.f, FallbackDuration, ActiveAnimSetIdx);
+	}
+
+	if (MissingRoles.Num() > 0)
+	{
+		FString RoleList;
+		for (const FName& Role : MissingRoles)
+		{
+			if (!RoleList.IsEmpty())
 			{
-				UE_LOG(LogContextualAnim, Warning, TEXT("Scene Asset: %s - Role Asset %s is missing Role %s. Please use Update Roles under settings, or update your roles asset to include this track."), 
-												*GetNameSafe(SceneAsset), *GetNameSafe(SceneAsset->GetRolesAsset()), *AnimTrack.Role.ToString());
+				RoleList += TEXT(", ");
 			}
+			RoleList += Role.ToString();
 		}
+
+		FMessageDialog::Open(EAppMsgType::Ok, FText::Format(
+			LOCTEXT("MissingRolesWarning",
+				"Scene Asset '{0}' has tracks for roles that no longer exist in Role Asset '{1}':\n\n{2}\n\n"
+				"These tracks will be skipped. Use 'Update Roles' in the Details panel to clean up stale roles."),
+			FText::FromString(GetNameSafe(SceneAsset)),
+			FText::FromString(GetNameSafe(SceneAsset->GetRolesAsset())),
+			FText::FromString(RoleList)));
 	}
 
 	Sequencer->NotifyMovieSceneDataChanged(EMovieSceneDataChangeType::MovieSceneStructureItemsChanged);
@@ -805,6 +868,7 @@ void FContextualAnimViewModel::StartSimulateMode()
 void FContextualAnimViewModel::StopSimulateMode()
 {
 	SimulateModeState = ESimulateModeState::Inactive;
+	CurrentSegmentIdx = INDEX_NONE;
 
 	Sequencer->Pause();
 	Sequencer->SetGlobalTime(0);
@@ -844,20 +908,250 @@ void FContextualAnimViewModel::ToggleSimulateMode()
 
 void FContextualAnimViewModel::StartSimulation()
 {
+	// For multi-section mode, start with the first segment's section
+	const int32 StartSectionIdx = HasMultiSectionPreviewTimelineData() ? ComputedTimeline[0].SectionIdx : ActiveSectionIdx;
+
 	TMap<FName, FContextualAnimSceneBindingContext> RoleToActorMap;
 	for (const FContextualAnimSceneBinding& Binding : SceneBindings)
 	{
 		RoleToActorMap.Add(SceneBindings.GetRoleFromBinding(Binding), Binding.GetContext());
 	}
 
-	if (FContextualAnimSceneBindings::TryCreateBindings(*GetSceneAsset(), ActiveSectionIdx, RoleToActorMap, SceneBindings))
+	bool bStarted = false;
+	if (FContextualAnimSceneBindings::TryCreateBindings(*GetSceneAsset(), StartSectionIdx, RoleToActorMap, SceneBindings))
 	{
 		const FContextualAnimSceneBinding* Leader = SceneBindings.GetSyncLeader();
 		if (Leader && Leader->GetSceneActorComponent()->StartContextualAnimScene(SceneBindings, {}))
 		{
 			SimulateModeState = ESimulateModeState::Playing;
+			bStarted = true;
+
+			if (HasMultiSectionPreviewTimelineData())
+			{
+				CurrentSegmentIdx = 0;
+			}
 		}
 	}
+
+	// @TODO: Temp solution for when the interaction fails to start during Simulation Mode
+	// TryCreateBindings should bubble up the error so we can show what went wrong to the user right here
+	if (!bStarted)
+	{
+		StopSimulateMode();
+
+		FMessageDialog::Open(EAppMsgType::Ok,
+			LOCTEXT("ContextualAnimSceneStartFailed",
+				"Contextual Anim Scene failed to start. Check the Output Log with 'LogContextualAnim' set to All for more details."));
+	}
+}
+
+void FContextualAnimViewModel::TickSimulateMode(float DeltaTime)
+{
+	if (SimulateModeState != ESimulateModeState::Playing || !HasMultiSectionPreviewTimelineData())
+	{
+		return;
+	}
+
+	if (!ComputedTimeline.IsValidIndex(CurrentSegmentIdx))
+	{
+		return;
+	}
+
+	const FContextualAnimSceneBinding* Leader = SceneBindings.GetSyncLeader();
+	if (!Leader)
+	{
+		return;
+	}
+
+	FAnimMontageInstance* MontageInstance = Leader->GetAnimMontageInstance();
+	if (!MontageInstance || !MontageInstance->Montage)
+	{
+		// Montage was removed (e.g. scene ended via LeaveScene on the last segment)
+		StopSimulateMode();
+		return;
+	}
+
+	const float Position = MontageInstance->GetPosition();
+	const float Length = MontageInstance->Montage->GetPlayLength();
+	const float BlendOutTime = MontageInstance->Montage->BlendOut.GetBlendTime();
+
+	// Trigger transition at the point where the montage starts blending out
+	const float TransitionPoint = Length - BlendOutTime;
+
+	if (Position >= TransitionPoint - KINDA_SMALL_NUMBER)
+	{
+		const int32 NextSegmentIdx = CurrentSegmentIdx + 1;
+		if (NextSegmentIdx < ComputedTimeline.Num())
+		{
+			const FPreviewTimelineSegment& NextSegment = ComputedTimeline[NextSegmentIdx];
+			const FContextualAnimSceneSection* Section = SceneAsset->GetSection(NextSegment.SectionIdx);
+			check(Section);
+
+			if (UContextualAnimSceneActorComponent* LeaderComp = Leader->GetSceneActorComponent())
+			{
+				if (!LeaderComp->TransitionContextualAnimSceneToSpecificSet(Section->GetName(), NextSegment.AnimSetIdx))
+				{
+					StopSimulateMode();
+
+					FMessageDialog::Open(EAppMsgType::Ok, FText::Format(
+						LOCTEXT("TransitionFailedDuringSimulate",
+							"Failed to transition to section '{0}' (AnimSet {1}) during simulation. "
+							"Simulation has been stopped. Check the Output Log for more details."),
+						FText::FromName(Section->GetName()),
+						FText::AsNumber(NextSegment.AnimSetIdx)));
+					return;
+				}
+
+				// Keep ViewModel's SceneBindings in sync with the component
+				SceneBindings = LeaderComp->GetBindings();
+			}
+
+			CurrentSegmentIdx = NextSegmentIdx;
+		}
+		// Last segment: let the montage end naturally (OnMontageBlendingOut -> LeaveScene)
+	}
+}
+
+bool FContextualAnimViewModel::TryBuildMultiSectionPreviewTimelineData()
+{
+	ComputedTimeline.Reset();
+	TotalSequenceDuration = 0.f;
+	CurrentSegmentIdx = INDEX_NONE;
+
+	if (!SceneAsset || SceneAsset->MultiSectionPreviewSequence.Num() == 0)
+	{
+		return false;
+	}
+
+	for (const FContextualAnimPreviewSequenceEntry& Entry : SceneAsset->MultiSectionPreviewSequence)
+	{
+		const int32 SectionIdx = SceneAsset->GetSectionIndex(Entry.SectionName);
+		if (SectionIdx == INDEX_NONE)
+		{
+			UE_LOG(LogContextualAnim, Warning, TEXT("TryBuildMultiSectionPreviewTimelineData: Section '%s' not found in asset, skipping."), *Entry.SectionName.ToString());
+			continue;
+		}
+
+		const FContextualAnimSceneSection* Section = SceneAsset->GetSection(SectionIdx);
+		check(Section);
+
+		// Determine which anim set to use
+		int32 ResolvedAnimSetIdx = Entry.AnimSetIdx;
+		if (ResolvedAnimSetIdx == INDEX_NONE)
+		{
+			// Use the editor's active anim set for this section
+			const int32* ActiveSetIdxPtr = ActiveAnimSetMap.Find(SectionIdx);
+			ResolvedAnimSetIdx = ActiveSetIdxPtr ? *ActiveSetIdxPtr : 0;
+		}
+
+		if (!Section->GetAnimSet(ResolvedAnimSetIdx))
+		{
+			UE_LOG(LogContextualAnim, Warning, TEXT("TryBuildMultiSectionPreviewTimelineData: AnimSet %d not valid for section '%s', skipping."), ResolvedAnimSetIdx, *Entry.SectionName.ToString());
+			continue;
+		}
+
+		// Compute duration as max animation length across all tracks in the resolved anim set
+		const FContextualAnimSet* AnimSet = Section->GetAnimSet(ResolvedAnimSetIdx);
+		float SegmentDuration = 0.f;
+		for (const FContextualAnimTrack& Track : AnimSet->Tracks)
+		{
+			if (Track.Animation)
+			{
+				SegmentDuration = FMath::Max(SegmentDuration, Track.Animation->GetPlayLength());
+			}
+		}
+		SegmentDuration = FMath::Max(SegmentDuration, KINDA_SMALL_NUMBER);
+
+		// Expand loops
+		for (int32 Loop = 0; Loop < Entry.LoopCount; Loop++)
+		{
+			FPreviewTimelineSegment Segment;
+			Segment.SectionIdx = SectionIdx;
+			Segment.AnimSetIdx = ResolvedAnimSetIdx;
+			Segment.LoopIteration = Loop;
+			Segment.GlobalStartTime = TotalSequenceDuration;
+			Segment.Duration = SegmentDuration;
+
+			ComputedTimeline.Add(Segment);
+			TotalSequenceDuration += SegmentDuration;
+		}
+	}
+
+	return ComputedTimeline.Num() > 0;
+}
+
+bool FContextualAnimViewModel::ResolveGlobalTime(float GlobalTime, int32& OutSectionIdx, int32& OutAnimSetIdx, float& OutLocalTime, int32& OutSegmentIdx) const
+{
+	if (ComputedTimeline.Num() == 0)
+	{
+		return false;
+	}
+
+	// Clamp to valid range
+	GlobalTime = FMath::Clamp(GlobalTime, 0.f, TotalSequenceDuration - KINDA_SMALL_NUMBER);
+
+	// Find which segment contains this global time
+	for (int32 Idx = 0; Idx < ComputedTimeline.Num(); Idx++)
+	{
+		const FPreviewTimelineSegment& Segment = ComputedTimeline[Idx];
+		const float SegmentEndTime = Segment.GlobalStartTime + Segment.Duration;
+
+		if (GlobalTime < SegmentEndTime || Idx == ComputedTimeline.Num() - 1)
+		{
+			OutSectionIdx = Segment.SectionIdx;
+			OutAnimSetIdx = Segment.AnimSetIdx;
+			OutLocalTime = FMath::Clamp(GlobalTime - Segment.GlobalStartTime, 0.f, Segment.Duration - KINDA_SMALL_NUMBER);
+			OutSegmentIdx = Idx;
+			return true;
+		}
+	}
+
+	return false;
+}
+
+void FContextualAnimViewModel::TransitionPreviewToSegment(int32 SegmentIdx)
+{
+	check(ComputedTimeline.IsValidIndex(SegmentIdx));
+
+	const FPreviewTimelineSegment& Segment = ComputedTimeline[SegmentIdx];
+	const FContextualAnimSceneSection* Section = SceneAsset->GetSection(Segment.SectionIdx);
+	check(Section);
+
+	// Use the runtime transition mechanism to swap montages
+	const FContextualAnimSceneBinding* Leader = SceneBindings.GetSyncLeader();
+	if (Leader)
+	{
+		if (UContextualAnimSceneActorComponent* LeaderComp = Leader->GetSceneActorComponent())
+		{
+			if (!LeaderComp->TransitionContextualAnimSceneToSpecificSet(Section->GetName(), Segment.AnimSetIdx))
+			{
+				FMessageDialog::Open(EAppMsgType::Ok, FText::Format(
+					LOCTEXT("TransitionFailedDuringScrub",
+						"Failed to transition to section '{0}' (AnimSet {1}) while scrubbing. "
+						"Check the Output Log for more details."),
+					FText::FromName(Section->GetName()),
+					FText::AsNumber(Segment.AnimSetIdx)));
+				return;
+			}
+
+			// Sync our SceneBindings from the component's updated copy
+			// so that subsequent reads (UpdatePreviewActorTransform, etc.)
+			// use the correct SectionIdx / AnimSetIdx / AnimTrack data.
+			SceneBindings = LeaderComp->GetBindings();
+		}
+	}
+
+	// Pause all montages and disable auto blend out (same as RefreshPreviewScene does after starting a scene)
+	for (const FContextualAnimSceneBinding& Binding : SceneBindings)
+	{
+		if (FAnimMontageInstance* MontageInstance = Binding.GetAnimMontageInstance())
+		{
+			MontageInstance->Pause();
+			MontageInstance->bEnableAutoBlendOut = false;
+		}
+	}
+
+	CurrentSegmentIdx = SegmentIdx;
 }
 
 UWorld* FContextualAnimViewModel::GetWorld() const
@@ -875,7 +1169,16 @@ void FContextualAnimViewModel::UpdatePreviewActorTransform(const FContextualAnim
 {
 	if (AActor* PreviewActor = Binding.GetActor())
 	{
-		FTransform Transform = SceneBindings.GetAnimTrackFromBinding(Binding).GetRootTransformAtTime(Time);
+		const FContextualAnimTrack& AnimTrack = SceneBindings.GetAnimTrackFromBinding(Binding);
+
+		// Skip position update for animations with ForceRootLock since they have no meaningful root motion data. 
+		// Preserves the actor's current position from the previous section.
+		if (HasForceRootLock(AnimTrack.Animation.Get(), Time))
+		{
+			return;
+		}
+
+		FTransform Transform = AnimTrack.GetRootTransformAtTime(Time);
 
 		// Special case for Character
 		if (ACharacter* PreviewCharacter = Cast<ACharacter>(PreviewActor))
@@ -941,14 +1244,45 @@ void FContextualAnimViewModel::SequencerTimeChanged()
 	EMovieScenePlayerStatus::Type CurrentStatus = Sequencer->GetPlaybackStatus();
 	const float CurrentSequencerTime = Sequencer->GetGlobalTime().AsSeconds();
 
-	for (const FContextualAnimSceneBinding& Binding : SceneBindings)
+	if (HasMultiSectionPreviewTimelineData())
 	{
-		if (FAnimMontageInstance* MontageInstance = Binding.GetAnimMontageInstance())
+		// Multi-section scrubbing: resolve global time to section-local coordinates
+		int32 SectionIdx, AnimSetIdx, SegmentIdx;
+		float LocalTime;
+
+		if (ResolveGlobalTime(CurrentSequencerTime, SectionIdx, AnimSetIdx, LocalTime, SegmentIdx))
 		{
-			const float AnimPlayLength = MontageInstance->Montage->GetPlayLength();
-			const float CurrentTime = FMath::Clamp(CurrentSequencerTime, 0.f, AnimPlayLength - KINDA_SMALL_NUMBER);
-			MontageInstance->SetPosition(CurrentTime);
-			UpdatePreviewActorTransform(Binding, CurrentTime);
+			// Transition to new segment if needed (handles both forward and backward scrubbing)
+			if (SegmentIdx != CurrentSegmentIdx)
+			{
+				TransitionPreviewToSegment(SegmentIdx);
+			}
+
+			// Set montage positions from resolved local time
+			for (const FContextualAnimSceneBinding& Binding : SceneBindings)
+			{
+				if (FAnimMontageInstance* MontageInstance = Binding.GetAnimMontageInstance())
+				{
+					const float AnimPlayLength = MontageInstance->Montage->GetPlayLength();
+					const float ClampedTime = FMath::Clamp(LocalTime, 0.f, AnimPlayLength - KINDA_SMALL_NUMBER);
+					MontageInstance->SetPosition(ClampedTime);
+					UpdatePreviewActorTransform(Binding, ClampedTime);
+				}
+			}
+		}
+	}
+	else
+	{
+		// Single-section scrubbing
+		for (const FContextualAnimSceneBinding& Binding : SceneBindings)
+		{
+			if (FAnimMontageInstance* MontageInstance = Binding.GetAnimMontageInstance())
+			{
+				const float AnimPlayLength = MontageInstance->Montage->GetPlayLength();
+				const float CurrentTime = FMath::Clamp(CurrentSequencerTime, 0.f, AnimPlayLength - KINDA_SMALL_NUMBER);
+				MontageInstance->SetPosition(CurrentTime);
+				UpdatePreviewActorTransform(Binding, CurrentTime);
+			}
 		}
 	}
 
@@ -1064,8 +1398,20 @@ void FContextualAnimViewModel::UpdateSelection(FName Role, int32 CriterionIdx, i
 	Sequencer->EmptySelection();
 
 	SelectionInfo.Role = Role;
-	SelectionInfo.SectionIdx = ActiveSectionIdx;
-	SelectionInfo.AnimSetIdx = ActiveAnimSetMap.FindRef(ActiveSectionIdx);
+
+	// In multi-section mode, use the current segment's section/animset (based on scrub position)
+	if (HasMultiSectionPreviewTimelineData() && ComputedTimeline.IsValidIndex(CurrentSegmentIdx))
+	{
+		const FPreviewTimelineSegment& Segment = ComputedTimeline[CurrentSegmentIdx];
+		SelectionInfo.SectionIdx = Segment.SectionIdx;
+		SelectionInfo.AnimSetIdx = Segment.AnimSetIdx;
+	}
+	else
+	{
+		SelectionInfo.SectionIdx = ActiveSectionIdx;
+		SelectionInfo.AnimSetIdx = ActiveAnimSetMap.FindRef(ActiveSectionIdx);
+	}
+
 	SelectionInfo.Criterion.Key = CriterionIdx;
 	SelectionInfo.Criterion.Value = CriterionDataIdx;
 
@@ -1141,13 +1487,23 @@ FText FContextualAnimViewModel::GetSelectionDebugText() const
 				{
 					const FTransform SocketTransform = MeshComp->GetSocketTransform(Params->SocketName, ERelativeTransformSpace::RTS_World);
 					const FTransform RelativeTransform = SelectedActor->GetTransform().GetRelativeTransform(SocketTransform);
-					AttachmentText = FString::Printf(TEXT("Location: %s Rotation: %s"), *RelativeTransform.GetLocation().ToString(), *RelativeTransform.Rotator().ToString());
+					const FVector Loc = RelativeTransform.GetLocation();
+					const FRotator Rot = RelativeTransform.Rotator();
+					AttachmentText = FString::Printf(TEXT("\n    Location: X=%.3f, Y=%.3f, Z=%.3f\n    Rotation: P=%.3f, Y=%.3f, R=%.3f"),
+						Loc.X, Loc.Y, Loc.Z, Rot.Pitch, Rot.Yaw, Rot.Roll);
 				}
 			}
 		}
 	}
-	return FText::FromString(FString::Printf(TEXT("Selection Info:\n Role: %s \n Actor: %s \n Attachment Relative Transform: %s \n Criterion: %d (%d)"),
-		*SelectionInfo.Role.ToString(), *GetNameSafe(SelectedActor), *AttachmentText, SelectionInfo.Criterion.Key, SelectionInfo.Criterion.Value));
+	FString ActiveSectionText = TEXT("None");
+	const int32 SectionIdx = GetActiveSection();
+	if (const FContextualAnimSceneSection* Section = SceneAsset->GetSection(SectionIdx))
+	{
+		ActiveSectionText = FString::Printf(TEXT("[%d] %s"), SectionIdx, *Section->GetName().ToString());
+	}
+
+	return FText::FromString(FString::Printf(TEXT("Active Section: %s\n\nSelection Info:\n Role: %s \n Attachment Relative Transform: %s \n Criterion: %d (%d)"),
+		*ActiveSectionText, *SelectionInfo.Role.ToString(), *AttachmentText, SelectionInfo.Criterion.Key, SelectionInfo.Criterion.Value));
 }
 
 bool FContextualAnimViewModel::ProcessInputDelta(FVector& InDrag, FRotator& InRot, FVector& InScale)
@@ -1495,8 +1851,9 @@ void FContextualAnimViewModel::OnFinishedChangingProperties(const FPropertyChang
 		return;
 	}
 
-	// Refresh preview scene when OverridePreviewData changes
-	if (MemberPropertyName == GET_MEMBER_NAME_CHECKED(UContextualAnimSceneAsset, OverridePreviewData))
+	// Refresh preview scene when Preview data changes
+	if (MemberPropertyName == GET_MEMBER_NAME_CHECKED(UContextualAnimSceneAsset, OverridePreviewData) || 
+		MemberPropertyName == GET_MEMBER_NAME_CHECKED(UContextualAnimSceneAsset, MultiSectionPreviewSequence))
 	{
 		SetDefaultMode();
 		return;
diff --git a/Engine/Plugins/Experimental/Animation/ContextualAnimation/Source/ContextualAnimationEditor/Private/DetailCustomizations/ContextualAnimSceneAssetDetailCustom.cpp b/Engine/Plugins/Experimental/Animation/ContextualAnimation/Source/ContextualAnimationEditor/Private/DetailCustomizations/ContextualAnimSceneAssetDetailCustom.cpp
--- a/Engine/Plugins/Experimental/Animation/ContextualAnimation/Source/ContextualAnimationEditor/Private/DetailCustomizations/ContextualAnimSceneAssetDetailCustom.cpp
+++ b/Engine/Plugins/Experimental/Animation/ContextualAnimation/Source/ContextualAnimationEditor/Private/DetailCustomizations/ContextualAnimSceneAssetDetailCustom.cpp
@@ -55,10 +55,13 @@ void FContextualAnimSceneAssetDetailCustom::CustomizeDetails(IDetailLayoutBuilde
 				})
 		];
 
+	IDetailCategoryBuilder& PreviewCategory = DetailBuilder.EditCategory(TEXT("Preview"));
+	PreviewCategory.SetSortOrder(1);
+
 	IDetailCategoryBuilder& Category = DetailBuilder.EditCategory(TEXT("Defaults"));
 	Category.SetCategoryVisibility(false);
 
-	const int32 SectionIdx = ViewModel->GetActiveSection();
+	const int32 SectionIdx = ViewModel->GetDetailsSectionIdx();
 	if (const FContextualAnimSceneSection* Section = SceneAsset->GetSection(SectionIdx))
 	{
 		// Create new category to show editable data only for the currently active section in the editor
@@ -114,7 +117,7 @@ void FContextualAnimSceneAssetDetailCustom::CustomizeDetails(IDetailLayoutBuilde
 			];
 
 		// Same as above but for the currently active AnimSet
-		const int32 AnimSetIdx = ViewModel->GetActiveAnimSetForSection(SectionIdx);
+		const int32 AnimSetIdx = ViewModel->GetDetailsAnimSetIdx();
 		if (const FContextualAnimSet* AnimSet = Section->GetAnimSet(AnimSetIdx))
 		{
 			// Create new category to show editable data only for the currently active AnimSet in the editor
diff --git a/Engine/Plugins/Experimental/Animation/ContextualAnimation/Source/ContextualAnimationEditor/Public/ContextualAnimMovieSceneSection.h b/Engine/Plugins/Experimental/Animation/ContextualAnimation/Source/ContextualAnimationEditor/Public/ContextualAnimMovieSceneSection.h
--- a/Engine/Plugins/Experimental/Animation/ContextualAnimation/Source/ContextualAnimationEditor/Public/ContextualAnimMovieSceneSection.h
+++ b/Engine/Plugins/Experimental/Animation/ContextualAnimation/Source/ContextualAnimationEditor/Public/ContextualAnimMovieSceneSection.h
@@ -15,7 +15,7 @@ class UContextualAnimMovieSceneSection : public UMovieSceneSection
 
 public:
 
-	void Initialize(int32 InSectionIdx, int32 InAnimSetIdx, int32 InAnimTrackIdx);
+	void Initialize(int32 InSectionIdx, int32 InAnimSetIdx, int32 InAnimTrackIdx, int32 InSegmentIdx = INDEX_NONE);
 
 	/** Returns the track owner of this section */
 	UContextualAnimMovieSceneTrack& GetOwnerTrack() const;
@@ -25,6 +25,7 @@ public:
 	FORCEINLINE int32 GetSectionIdx() const { return SectionIdx; }
 	FORCEINLINE int32 GetAnimSetIdx() const { return AnimSetIdx; }
 	FORCEINLINE int32 GetAnimTrackIdx() const { return AnimTrackIdx; }
+	FORCEINLINE int32 GetSegmentIdx() const { return SegmentIdx; }
 
 private:
 
@@ -36,4 +37,8 @@ private:
 
 	UPROPERTY()
 	int32 AnimTrackIdx = INDEX_NONE;
+
+	/** Index into the computed timeline segments (INDEX_NONE for single-section mode) */
+	UPROPERTY()
+	int32 SegmentIdx = INDEX_NONE;
 };
diff --git a/Engine/Plugins/Experimental/Animation/ContextualAnimation/Source/ContextualAnimationEditor/Public/ContextualAnimMovieSceneTrackEditor.h b/Engine/Plugins/Experimental/Animation/ContextualAnimation/Source/ContextualAnimationEditor/Public/ContextualAnimMovieSceneTrackEditor.h
--- a/Engine/Plugins/Experimental/Animation/ContextualAnimation/Source/ContextualAnimationEditor/Public/ContextualAnimMovieSceneTrackEditor.h
+++ b/Engine/Plugins/Experimental/Animation/ContextualAnimation/Source/ContextualAnimationEditor/Public/ContextualAnimMovieSceneTrackEditor.h
@@ -21,6 +21,7 @@ public:
 	virtual TSharedPtr<SWidget> BuildOutlinerEditWidget(const FGuid& ObjectBinding, UMovieSceneTrack* Track, const FBuildEditWidgetParams& Params) override;
 	virtual void BuildAddTrackMenu(FMenuBuilder& MenuBuilder) override;
 	virtual void BuildObjectBindingTrackMenu(FMenuBuilder& MenuBuilder, const TArray<FGuid>& ObjectBindings, const UClass* ObjectClass) override;
+	virtual bool GetDefaultExpansionState(UMovieSceneTrack* InTrack) const override { return true; }
 
 	static TSharedRef<ISequencerTrackEditor> CreateTrackEditor(TSharedRef<ISequencer> InSequencer);
 };
@@ -37,6 +38,7 @@ public:
 		: FSequencerSection(InSectionObject)
 	{}
 
+	virtual int32 OnPaintSection(FSequencerSectionPainter& InPainter) const override;
 	virtual FReply OnSectionDoubleClicked(const FGeometry& SectionGeometry, const FPointerEvent& MouseEvent) override;
 	virtual FText GetSectionTitle() const override;
 	virtual float GetSectionHeight() const override;
diff --git a/Engine/Plugins/Experimental/Animation/ContextualAnimation/Source/ContextualAnimationEditor/Public/ContextualAnimViewModel.h b/Engine/Plugins/Experimental/Animation/ContextualAnimation/Source/ContextualAnimationEditor/Public/ContextualAnimViewModel.h
--- a/Engine/Plugins/Experimental/Animation/ContextualAnimation/Source/ContextualAnimationEditor/Public/ContextualAnimViewModel.h
+++ b/Engine/Plugins/Experimental/Animation/ContextualAnimation/Source/ContextualAnimationEditor/Public/ContextualAnimViewModel.h
@@ -31,6 +31,16 @@ struct FMovieSceneSectionMovedParams;
 struct FContextualAnimNewAnimSetParams;
 struct FContextualAnimTrack;
 
+/** Represents a single segment in the flattened multi-section preview timeline */
+struct FPreviewTimelineSegment
+{
+	int32 SectionIdx;           // Index into SceneAsset->Sections
+	int32 AnimSetIdx;           // Which anim set to use for this segment
+	int32 LoopIteration;        // 0-based loop iteration counter
+	float GlobalStartTime;      // Start time in the flattened timeline
+	float Duration;             // Max animation length across all tracks in this anim set
+};
+
 class FContextualAnimViewModel : public TSharedFromThis<FContextualAnimViewModel>, public FGCObject, public FSelfRegisteringEditorUndoClient
 {
 public:
@@ -95,6 +105,11 @@ public:
 	void SetActiveAnimSetForSection(int32 SectionIdx, int32 AnimSetIdx);
 	int32 GetActiveAnimSetForSection(int32 SectionIdx) const;
 
+	/** Returns the section index from the current selection, or the active section if nothing is selected */
+	int32 GetDetailsSectionIdx() const;
+	/** Returns the anim set index from the current selection, or the active anim set if nothing is selected */
+	int32 GetDetailsAnimSetIdx() const;
+
 	void StartSimulateMode();
 	void StopSimulateMode();
 	void ToggleSimulateMode();
@@ -103,6 +118,21 @@ public:
 	bool IsSimulateModePlaying()	const { return SimulateModeState == ESimulateModeState::Playing; }
 	void StartSimulation();
 
+	/** Called from FContextualAnimEdMode::Tick to poll montage state during multi-section simulate mode */
+	void TickSimulateMode(float DeltaTime);
+
+	/** Attempts to build the flattened multi-section preview timeline from the asset's PreviewSequence entries. Returns true if segments were created. */
+	bool TryBuildMultiSectionPreviewTimelineData();
+
+	/** Maps a global timeline position to section-local coordinates. Returns false if GlobalTime is out of range. */
+	bool ResolveGlobalTime(float GlobalTime, int32& OutSectionIdx, int32& OutAnimSetIdx, float& OutLocalTime, int32& OutSegmentIdx) const;
+
+	/** Returns true if the multi-section preview timeline is active (has segments) */
+	bool HasMultiSectionPreviewTimelineData() const { return ComputedTimeline.Num() > 0; }
+
+	/** Transitions preview to a different segment during scrubbing */
+	void TransitionPreviewToSegment(int32 SegmentIdx);
+
 	void UpdatePreviewActorTransform(const FContextualAnimSceneBinding& Binding, float Time);
 
 	void UpdateSelection(const AActor* SelectedActor);
@@ -154,6 +184,8 @@ private:
 		void Reset()
 		{
 			Role = NAME_None;
+			SectionIdx = INDEX_NONE;
+			AnimSetIdx = INDEX_NONE;
 			Criterion.Key = INDEX_NONE;
 			Criterion.Value = INDEX_NONE;
 			NotifyGuid = FGuid();
@@ -191,6 +223,15 @@ private:
 	/** The previous time for the sequencer time line. */
 	float PreviousSequencerTime = 0.f;
 
+	/** Flattened multi-section preview timeline segments */
+	TArray<FPreviewTimelineSegment> ComputedTimeline;
+
+	/** Total duration of the flattened preview sequence */
+	float TotalSequenceDuration = 0.f;
+
+	/** Current segment index during scrub/simulate in multi-section mode */
+	int32 CurrentSegmentIdx = INDEX_NONE;
+
 	/** Flag for preventing OnAnimNotifyChanged from updating tracks when the change to the animation came from us */
 	bool bUpdatingAnimationFromSequencer = false;
 
