[CAS] Override Support for CAS Actors
- Give Override Provider an Option to Adjust the Preview

#rnx
#Synced-CL 52452884

diff --git a/Engine/Plugins/Experimental/Animation/ContextualAnimation/Source/ContextualAnimation/Public/ContextualAnimOverrideInterface.h b/Engine/Plugins/Experimental/Animation/ContextualAnimation/Source/ContextualAnimation/Public/ContextualAnimOverrideInterface.h
--- a/Engine/Plugins/Experimental/Animation/ContextualAnimation/Source/ContextualAnimation/Public/ContextualAnimOverrideInterface.h
+++ b/Engine/Plugins/Experimental/Animation/ContextualAnimation/Source/ContextualAnimation/Public/ContextualAnimOverrideInterface.h
@@ -9,6 +9,7 @@
 #include "ContextualAnimOverrideInterface.generated.h"
 
 class UAnimSequenceBase;
+struct FContextualAnimActorPreviewData;
 
 
 USTRUCT(BlueprintType)
@@ -55,4 +56,17 @@ class IContextualAnimOverrideInterface : public IInterface
 public:
 	virtual bool GetCASAnimationDataOverride(const FContextualAnimSceneBinding& InBinding, const FContextualAnimSceneBindings& InBindings, const FContextualAnimTrack& InAnimTrack, FContextualAnimationOverrideData& OutData) = 0;
 
+#if WITH_EDITORONLY_DATA
+	/**
+	 * Called during SpawnPreviewActor to allow the provider to override the preview data for a given role.
+	 * Return true and populate OutPreviewData to replace whatever is set in the scene asset's OverridePreviewData.
+	 */
+	virtual bool GetCASAnimActorPreviewDataOverrideForRole(const FName& InRole, FContextualAnimActorPreviewData& OutPreviewData) { return false; }
+
+	/**
+	 * Called after a preview actor has been spawned, allowing post-spawn setup (e.g. applying cosmetics).
+	 */
+	virtual void OnCASPreviewActorSpawned(const FName& InRole, AActor* InSpawnedActor) {}
+#endif
+
 };
diff --git a/Engine/Plugins/Experimental/Animation/ContextualAnimation/Source/ContextualAnimationEditor/Private/ContextualAnimViewModel.cpp b/Engine/Plugins/Experimental/Animation/ContextualAnimation/Source/ContextualAnimationEditor/Private/ContextualAnimViewModel.cpp
--- a/Engine/Plugins/Experimental/Animation/ContextualAnimation/Source/ContextualAnimationEditor/Private/ContextualAnimViewModel.cpp
+++ b/Engine/Plugins/Experimental/Animation/ContextualAnimation/Source/ContextualAnimationEditor/Private/ContextualAnimViewModel.cpp
@@ -243,8 +243,17 @@ AActor* FContextualAnimViewModel::SpawnPreviewActor(const FContextualAnimTrack&
 	UClass* PreviewActorClass = nullptr;
 	UClass* PreviewAnimInstanceClass = nullptr;
 
-	for (const FContextualAnimActorPreviewData& PreviewData : GetSceneAsset()->OverridePreviewData)
+	IContextualAnimOverrideInterface* OverrideInterface = SceneAsset->PreviewOverrideProvider ?  Cast<IContextualAnimOverrideInterface>(SceneAsset->PreviewOverrideProvider->GetDefaultObject()) : nullptr;
+	bool bPreviewOveridden = false;
+
+	for (FContextualAnimActorPreviewData PreviewData : GetSceneAsset()->OverridePreviewData)
 	{
+		// Allow the PreviewOverrideProvider to override the preview data for this role
+		if (OverrideInterface && OverrideInterface->GetCASAnimActorPreviewDataOverrideForRole(AnimTrack.Role, PreviewData))
+		{
+			bPreviewOveridden = true;
+		}
+
 		if (PreviewData.Role == AnimTrack.Role)
 		{
 			if (PreviewData.Type == EContextualAnimActorPreviewType::StaticMesh)
@@ -301,7 +310,7 @@ AActor* FContextualAnimViewModel::SpawnPreviewActor(const FContextualAnimTrack&
 		const FContextualAnimRoleDefinition* RoleDef = GetSceneAsset()->RolesAsset ? GetSceneAsset()->RolesAsset->FindRoleDefinitionByName(AnimTrack.Role) : nullptr;
 		const bool bIsCharacter = (RoleDef && RoleDef->bIsCharacter);
 
-		if (bIsCharacter)
+		if (bIsCharacter && !bPreviewOveridden)
 		{
 			ACharacter* PreviewCharacter = GetWorld()->SpawnActor<ACharacter>(ACharacter::StaticClass(), SpawnTransform, Params);
 			PreviewCharacter->SetFlags(RF_Transient);
@@ -356,6 +365,11 @@ AActor* FContextualAnimViewModel::SpawnPreviewActor(const FContextualAnimTrack&
 				PreviewActor = GetWorld()->SpawnActor<AActor>(PreviewActorClass, SpawnTransform, Params);
 				PreviewActor->SetFlags(RF_Transient);
 
+				// Allow the Override Interface to Access the PreviewActor before calling refresh, this can set any data needed to properly preview the scene with this override set
+				if (OverrideInterface)
+				{
+					OverrideInterface->OnCASPreviewActorSpawned(AnimTrack.Role, PreviewActor);
+				}
 				// Try to call RefreshEditorPreview on the actor. This give game specific actors a chance to run any logic needed to be able to use them as preview actor in this editor.
 				if (PreviewActor->GetClass()->ImplementsInterface(UContextualAnimActorInterface::StaticClass()))
 				{
